[NO NON-PII COMMIT MESSAGE]

diff --git a/Templates/TP_FirstPerson/Source/TP_FirstPerson/TP_FirstPersonCharacter.cpp b/Templates/TP_FirstPerson/Source/TP_FirstPerson/TP_FirstPersonCharacter.cpp
--- a/Templates/TP_FirstPerson/Source/TP_FirstPerson/TP_FirstPersonCharacter.cpp
+++ b/Templates/TP_FirstPerson/Source/TP_FirstPerson/TP_FirstPersonCharacter.cpp
@@ -153,16 +153,16 @@ void ATP_FirstPersonCharacter::TouchUpdate(const ETouchIndex::Type FingerIndex,
 					FVector2D ScreenSize;
 					ViewportClient->GetViewportSize(ScreenSize);
 					FVector2D ScaledDelta = FVector2D( MoveDelta.X, MoveDelta.Y) / ScreenSize;									
-					if (ScaledDelta.X != 0.0f)
+					if (FMath::Abs(ScaledDelta.X) >= 4.0 / ScreenSize.X)
 					{
 						TouchItem.bMoved = true;
 						float Value = ScaledDelta.X * BaseTurnRate;
 						AddControllerYawInput(Value);
 					}
-					if (ScaledDelta.Y != 0.0f)
+					if (FMath::Abs(ScaledDelta.Y) >= 4.0 / ScreenSize.Y)
 					{
 						TouchItem.bMoved = true;
-						float Value = ScaledDelta.Y* BaseTurnRate;
+						float Value = ScaledDelta.Y * BaseTurnRate;
 						AddControllerPitchInput(Value);
 					}
 					TouchItem.Location = Location;
