Fixing bullet spawn location with a USceneComponent to imitate the fix made to the BP version.

diff --git a/Templates/TP_FirstPerson/Source/TP_FirstPerson/TP_FirstPersonCharacter.cpp b/Templates/TP_FirstPerson/Source/TP_FirstPerson/TP_FirstPersonCharacter.cpp
--- a/Templates/TP_FirstPerson/Source/TP_FirstPerson/TP_FirstPersonCharacter.cpp
+++ b/Templates/TP_FirstPerson/Source/TP_FirstPerson/TP_FirstPersonCharacter.cpp
@@ -40,6 +40,9 @@ ATP_FirstPersonCharacter::ATP_FirstPersonCharacter()
 	FP_Gun->CastShadow = false;
 	FP_Gun->AttachTo(Mesh1P, TEXT("GripPoint"), EAttachLocation::SnapToTargetIncludingScale, true);
 
+	FP_MuzzleLocation = CreateDefaultSubobject<USceneComponent>(TEXT("MuzzleLocation"));
+	FP_MuzzleLocation->AttachTo(FP_Gun);
+	FP_MuzzleLocation->SetRelativeLocation(FVector(0.2f, 48.4f, -10.6f));
 
 	// Default offset from the character location for projectiles to spawn
 	GunOffset = FVector(100.0f, 30.0f, 10.0f);
@@ -84,7 +87,7 @@ void ATP_FirstPersonCharacter::OnFire()
 	{
 		const FRotator SpawnRotation = GetControlRotation();
 		// MuzzleOffset is in camera space, so transform it to world space before offsetting from the character location to find the final muzzle position
-		const FVector SpawnLocation = GetActorLocation() + SpawnRotation.RotateVector(GunOffset);
+		const FVector SpawnLocation = ((FP_MuzzleLocation != nullptr) ? FP_MuzzleLocation->GetComponentLocation() : GetActorLocation()) + SpawnRotation.RotateVector(GunOffset);
 
 		UWorld* const World = GetWorld();
 		if (World != NULL)
diff --git a/Templates/TP_FirstPerson/Source/TP_FirstPerson/TP_FirstPersonCharacter.h b/Templates/TP_FirstPerson/Source/TP_FirstPerson/TP_FirstPersonCharacter.h
--- a/Templates/TP_FirstPerson/Source/TP_FirstPerson/TP_FirstPersonCharacter.h
+++ b/Templates/TP_FirstPerson/Source/TP_FirstPerson/TP_FirstPersonCharacter.h
@@ -18,6 +18,10 @@ class ATP_FirstPersonCharacter : public ACharacter
 	UPROPERTY(VisibleDefaultsOnly, Category = Mesh)
 	class USkeletalMeshComponent* FP_Gun;
 
+	/** Location on gun mesh where projectiles should spawn. */
+	UPROPERTY(VisibleDefaultsOnly, Category = Mesh)
+	class USceneComponent* FP_MuzzleLocation;
+
 	/** First person camera */
 	UPROPERTY(VisibleAnywhere, BlueprintReadOnly, Category = Camera, meta = (AllowPrivateAccess = "true"))
 	class UCameraComponent* FirstPersonCameraComponent;
