Updating First Person C++ template to include the same functionality as the BP version. Character BP not updated to include the VR Gun model.

diff --git a/Templates/TP_FirstPerson/Config/DefaultInput.ini b/Templates/TP_FirstPerson/Config/DefaultInput.ini
--- a/Templates/TP_FirstPerson/Config/DefaultInput.ini
+++ b/Templates/TP_FirstPerson/Config/DefaultInput.ini
@@ -1,19 +1,25 @@
 [/Script/Engine.InputSettings]
 +ActionMappings=(ActionName="Jump", Key=SpaceBar)
 +ActionMappings=(ActionName="Jump", Key=Gamepad_FaceButton_Bottom)
++ActionMappings=(ActionName="Jump",Key=MotionController_Left_Trigger,bShift=False,bCtrl=False,bAlt=False,bCmd=False)
 
 +ActionMappings=(ActionName="Fire", Key=LeftMouseButton)
 +ActionMappings=(ActionName="Fire", Key=Gamepad_RightTrigger)
++ActionMappings=(ActionName="Fire",Key=MotionController_Right_Trigger,bShift=False,bCtrl=False,bAlt=False,bCmd=False)
+
++ActionMappings=(ActionName="ResetVR",Key=R,bShift=False,bCtrl=False,bAlt=False,bCmd=False)
 
 +AxisMappings=(AxisName="MoveForward", Key=W, Scale=1.f)
 +AxisMappings=(AxisName="MoveForward", Key=S, Scale=-1.f)
 +AxisMappings=(AxisName="MoveForward", Key=Up, Scale=1.f)
 +AxisMappings=(AxisName="MoveForward", Key=Down, Scale=-1.f)
 +AxisMappings=(AxisName="MoveForward", Key=Gamepad_LeftY, Scale=1.f)
++AxisMappings=(AxisName="MoveForward",Key=MotionController_Left_Thumbstick_Y,Scale=-1.000000)
 
 +AxisMappings=(AxisName="MoveRight", Key=A, Scale=-1.f)
 +AxisMappings=(AxisName="MoveRight", Key=D, Scale=1.f)
 +AxisMappings=(AxisName="MoveRight", Key=Gamepad_LeftX, Scale=1.f)
++AxisMappings=(AxisName="MoveRight",Key=MotionController_Left_Thumbstick_X,Scale=1.000000)
 
 +AxisMappings=(AxisName="TurnRate", Key=Gamepad_RightX, Scale=1.f)
 +AxisMappings=(AxisName="TurnRate", Key=Left, Scale=-1.f)
diff --git a/Templates/TP_FirstPerson/Source/TP_FirstPerson/TP_FirstPerson.Build.cs b/Templates/TP_FirstPerson/Source/TP_FirstPerson/TP_FirstPerson.Build.cs
--- a/Templates/TP_FirstPerson/Source/TP_FirstPerson/TP_FirstPerson.Build.cs
+++ b/Templates/TP_FirstPerson/Source/TP_FirstPerson/TP_FirstPerson.Build.cs
@@ -6,6 +6,6 @@ public class TP_FirstPerson : ModuleRules
 {
 	public TP_FirstPerson(TargetInfo Target)
 	{
-		PublicDependencyModuleNames.AddRange(new string[] { "Core", "CoreUObject", "Engine", "InputCore" });
+		PublicDependencyModuleNames.AddRange(new string[] { "Core", "CoreUObject", "Engine", "InputCore", "HeadMountedDisplay" });
 	}
 }
diff --git a/Templates/TP_FirstPerson/Source/TP_FirstPerson/TP_FirstPersonCharacter.cpp b/Templates/TP_FirstPerson/Source/TP_FirstPerson/TP_FirstPersonCharacter.cpp
--- a/Templates/TP_FirstPerson/Source/TP_FirstPerson/TP_FirstPersonCharacter.cpp
+++ b/Templates/TP_FirstPerson/Source/TP_FirstPerson/TP_FirstPersonCharacter.cpp
@@ -5,6 +5,8 @@
 #include "TP_FirstPersonProjectile.h"
 #include "Animation/AnimInstance.h"
 #include "GameFramework/InputSettings.h"
+#include "Kismet/HeadMountedDisplayFunctionLibrary.h"
+#include "MotionControllerComponent.h"
 
 DEFINE_LOG_CATEGORY_STATIC(LogFPChar, Warning, All);
 
@@ -41,6 +43,7 @@ ATP_FirstPersonCharacter::ATP_FirstPersonCharacter()
 	FP_Gun->bCastDynamicShadow = false;
 	FP_Gun->CastShadow = false;
 	// FP_Gun->SetupAttachment(Mesh1P, TEXT("GripPoint"));
+	FP_Gun->SetupAttachment(RootComponent);
 
 	FP_MuzzleLocation = CreateDefaultSubobject<USceneComponent>(TEXT("MuzzleLocation"));
 	FP_MuzzleLocation->SetupAttachment(FP_Gun);
@@ -49,8 +52,32 @@ ATP_FirstPersonCharacter::ATP_FirstPersonCharacter()
 	// Default offset from the character location for projectiles to spawn
 	GunOffset = FVector(100.0f, 30.0f, 10.0f);
 
-	// Note: The ProjectileClass and the skeletal mesh/anim blueprints for Mesh1P are set in the
-	// derived blueprint asset named MyCharacter (to avoid direct content references in C++)
+	// Note: The ProjectileClass and the skeletal mesh/anim blueprints for Mesh1P, FP_Gun, and VR_Gun 
+	// are set in the derived blueprint asset named MyCharacter to avoid direct content references in C++.
+
+	// Create VR Controllers.
+	R_MotionController = CreateDefaultSubobject<UMotionControllerComponent>(TEXT("R_MotionController"));
+	R_MotionController->Hand = EControllerHand::Right;
+	R_MotionController->SetupAttachment(RootComponent);
+	L_MotionController = CreateDefaultSubobject<UMotionControllerComponent>(TEXT("L_MotionController"));
+	L_MotionController->SetupAttachment(RootComponent);
+
+	// Create a gun and attach it to the right-hand VR controller.
+	// Create a gun mesh component
+	VR_Gun = CreateDefaultSubobject<USkeletalMeshComponent>(TEXT("VR_Gun"));
+	VR_Gun->SetOnlyOwnerSee(true);			// only the owning player will see this mesh
+	VR_Gun->bCastDynamicShadow = false;
+	VR_Gun->CastShadow = false;
+	VR_Gun->SetupAttachment(R_MotionController);
+	VR_Gun->SetRelativeRotation(FRotator(0.0f, -90.0f, 0.0f));
+
+	VR_MuzzleLocation = CreateDefaultSubobject<USceneComponent>(TEXT("VR_MuzzleLocation"));
+	VR_MuzzleLocation->SetupAttachment(VR_Gun);
+	VR_MuzzleLocation->SetRelativeLocation(FVector(0.000004, 53.999992, 10.000000));
+	VR_MuzzleLocation->SetRelativeRotation(FRotator(0.0f, 90.0f, 0.0f));		// Counteract the rotation of the VR gun model.
+
+	// Uncomment the following line to turn motion controllers on by default:
+	//bUsingMotionControllers = true;
 }
 
 void ATP_FirstPersonCharacter::BeginPlay()
@@ -58,7 +85,18 @@ void ATP_FirstPersonCharacter::BeginPlay()
 	// Call the base class  
 	Super::BeginPlay();
 
-	FP_Gun->AttachToComponent(Mesh1P, FAttachmentTransformRules(EAttachmentRule::SnapToTarget, true), TEXT("GripPoint")); //Attach gun mesh component to Skeleton, doing it here because the skelton is not yet created in the constructor
+	//Attach gun mesh component to Skeleton, doing it here because the skeleton is not yet created in the constructor
+	FP_Gun->AttachToComponent(Mesh1P, FAttachmentTransformRules(EAttachmentRule::SnapToTarget, true), TEXT("GripPoint"));
+
+	// Hide whichever version of the gun we're not using.
+	if (bUsingMotionControllers)
+	{
+		Mesh1P->SetHiddenInGame(true, true);
+	}
+	else
+	{
+		VR_Gun->SetHiddenInGame(true, true);
+	}
 }
 
 //////////////////////////////////////////////////////////////////////////
@@ -78,6 +116,8 @@ void ATP_FirstPersonCharacter::SetupPlayerInputComponent(class UInputComponent*
 		PlayerInputComponent->BindAction("Fire", IE_Pressed, this, &ATP_FirstPersonCharacter::OnFire);
 	}
 
+	PlayerInputComponent->BindAction("ResetVR", IE_Pressed, this, &ATP_FirstPersonCharacter::OnResetVR);
+
 	PlayerInputComponent->BindAxis("MoveForward", this, &ATP_FirstPersonCharacter::MoveForward);
 	PlayerInputComponent->BindAxis("MoveRight", this, &ATP_FirstPersonCharacter::MoveRight);
 
@@ -95,15 +135,24 @@ void ATP_FirstPersonCharacter::OnFire()
 	// try and fire a projectile
 	if (ProjectileClass != NULL)
 	{
-		const FRotator SpawnRotation = GetControlRotation();
-		// MuzzleOffset is in camera space, so transform it to world space before offsetting from the character location to find the final muzzle position
-		const FVector SpawnLocation = ((FP_MuzzleLocation != nullptr) ? FP_MuzzleLocation->GetComponentLocation() : GetActorLocation()) + SpawnRotation.RotateVector(GunOffset);
-
 		UWorld* const World = GetWorld();
 		if (World != NULL)
 		{
-			// spawn the projectile at the muzzle
-			World->SpawnActor<ATP_FirstPersonProjectile>(ProjectileClass, SpawnLocation, SpawnRotation);
+			if (bUsingMotionControllers)
+			{
+				const FRotator SpawnRotation = VR_MuzzleLocation->GetComponentRotation();
+				const FVector SpawnLocation = VR_MuzzleLocation->GetComponentLocation();
+				World->SpawnActor<ATP_FirstPersonProjectile>(ProjectileClass, SpawnLocation, SpawnRotation);
+			}
+			else
+			{
+				const FRotator SpawnRotation = GetControlRotation();
+				// MuzzleOffset is in camera space, so transform it to world space before offsetting from the character location to find the final muzzle position
+				const FVector SpawnLocation = ((FP_MuzzleLocation != nullptr) ? FP_MuzzleLocation->GetComponentLocation() : GetActorLocation()) + SpawnRotation.RotateVector(GunOffset);
+
+				// spawn the projectile at the muzzle
+				World->SpawnActor<ATP_FirstPersonProjectile>(ProjectileClass, SpawnLocation, SpawnRotation);
+			}
 		}
 	}
 
@@ -123,7 +172,11 @@ void ATP_FirstPersonCharacter::OnFire()
 			AnimInstance->Montage_Play(FireAnimation, 1.f);
 		}
 	}
+}
 
+void ATP_FirstPersonCharacter::OnResetVR()
+{
+	UHeadMountedDisplayFunctionLibrary::ResetOrientationAndPosition();
 }
 
 void ATP_FirstPersonCharacter::BeginTouch(const ETouchIndex::Type FingerIndex, const FVector Location)
diff --git a/Templates/TP_FirstPerson/Source/TP_FirstPerson/TP_FirstPersonCharacter.h b/Templates/TP_FirstPerson/Source/TP_FirstPerson/TP_FirstPersonCharacter.h
--- a/Templates/TP_FirstPerson/Source/TP_FirstPerson/TP_FirstPersonCharacter.h
+++ b/Templates/TP_FirstPerson/Source/TP_FirstPerson/TP_FirstPersonCharacter.h
@@ -22,9 +22,25 @@ class ATP_FirstPersonCharacter : public ACharacter
 	UPROPERTY(VisibleDefaultsOnly, Category = Mesh)
 	class USceneComponent* FP_MuzzleLocation;
 
+	/** Gun mesh: VR view (attached to the VR controller directly, no arm, just the actual gun) */
+	UPROPERTY(VisibleDefaultsOnly, Category = Mesh)
+	class USkeletalMeshComponent* VR_Gun;
+
+	/** Location on VR gun mesh where projectiles should spawn. */
+	UPROPERTY(VisibleDefaultsOnly, Category = Mesh)
+	class USceneComponent* VR_MuzzleLocation;
+
 	/** First person camera */
 	UPROPERTY(VisibleAnywhere, BlueprintReadOnly, Category = Camera, meta = (AllowPrivateAccess = "true"))
 	class UCameraComponent* FirstPersonCameraComponent;
+
+	/** Motion controller (right hand) */
+	UPROPERTY(VisibleAnywhere, BlueprintReadOnly, meta = (AllowPrivateAccess = "true"))
+	class UMotionControllerComponent* R_MotionController;
+
+	/** Motion controller (left hand) */
+	UPROPERTY(VisibleAnywhere, BlueprintReadOnly, meta = (AllowPrivateAccess = "true"))
+	class UMotionControllerComponent* L_MotionController;
 public:
 	ATP_FirstPersonCharacter();
 
@@ -54,11 +70,18 @@ public:
 	UPROPERTY(EditAnywhere, BlueprintReadWrite, Category = Gameplay)
 	class UAnimMontage* FireAnimation;
 
+	/** Whether to use motion controller location for aiming. */
+	UPROPERTY(EditAnywhere, BlueprintReadWrite, Category = Gameplay)
+	uint32 bUsingMotionControllers : 1;
+
 protected:
 	
 	/** Fires a projectile. */
 	void OnFire();
 
+	/** Resets HMD orientation and position in VR. */
+	void OnResetVR();
+
 	/** Handles moving forward/backward */
 	void MoveForward(float Val);
 
