Improving C++ implementation with a new BP that includes the gun mesh.

diff --git a/Templates/TP_FirstPerson/Source/TP_FirstPerson/TP_FirstPersonCharacter.cpp b/Templates/TP_FirstPerson/Source/TP_FirstPerson/TP_FirstPersonCharacter.cpp
--- a/Templates/TP_FirstPerson/Source/TP_FirstPerson/TP_FirstPersonCharacter.cpp
+++ b/Templates/TP_FirstPerson/Source/TP_FirstPerson/TP_FirstPersonCharacter.cpp
@@ -88,14 +88,16 @@ void ATP_FirstPersonCharacter::BeginPlay()
 	//Attach gun mesh component to Skeleton, doing it here because the skeleton is not yet created in the constructor
 	FP_Gun->AttachToComponent(Mesh1P, FAttachmentTransformRules(EAttachmentRule::SnapToTarget, true), TEXT("GripPoint"));
 
-	// Hide whichever version of the gun we're not using.
+	// Show or hide the two versions of the gun based on whether or not we're using motion controllers.
 	if (bUsingMotionControllers)
 	{
+		VR_Gun->SetHiddenInGame(false, true);
 		Mesh1P->SetHiddenInGame(true, true);
 	}
 	else
 	{
 		VR_Gun->SetHiddenInGame(true, true);
+		Mesh1P->SetHiddenInGame(false, true);
 	}
 }
 
