Adjusted projectile spawn location by adjusting WeaponOffset Y value in both code and BP template.

diff --git a/Templates/TP_FirstPerson/Source/TP_FirstPerson/TP_FirstPersonCharacter.cpp b/Templates/TP_FirstPerson/Source/TP_FirstPerson/TP_FirstPersonCharacter.cpp
--- a/Templates/TP_FirstPerson/Source/TP_FirstPerson/TP_FirstPersonCharacter.cpp
+++ b/Templates/TP_FirstPerson/Source/TP_FirstPerson/TP_FirstPersonCharacter.cpp
@@ -50,7 +50,7 @@ ATP_FirstPersonCharacter::ATP_FirstPersonCharacter()
 	FP_MuzzleLocation->SetRelativeLocation(FVector(0.2f, 48.4f, -10.6f));
 
 	// Default offset from the character location for projectiles to spawn
-	GunOffset = FVector(100.0f, 30.0f, 10.0f);
+	GunOffset = FVector(100.0f, 0.0f, 10.0f);
 
 	// Note: The ProjectileClass and the skeletal mesh/anim blueprints for Mesh1P, FP_Gun, and VR_Gun 
 	// are set in the derived blueprint asset named MyCharacter to avoid direct content references in C++.
diff --git a/Templates/TP_FirstPerson/Source/TP_FirstPerson/TP_FirstPersonHUD.cpp b/Templates/TP_FirstPerson/Source/TP_FirstPerson/TP_FirstPersonHUD.cpp
--- a/Templates/TP_FirstPerson/Source/TP_FirstPerson/TP_FirstPersonHUD.cpp
+++ b/Templates/TP_FirstPerson/Source/TP_FirstPerson/TP_FirstPersonHUD.cpp
@@ -25,7 +25,7 @@ void ATP_FirstPersonHUD::DrawHUD()
 
 	// offset by half the texture's dimensions so that the center of the texture aligns with the center of the Canvas
 	const FVector2D CrosshairDrawPosition( (Center.X),
-										   (Center.Y));
+										   (Center.Y + 20.0f));
 
 	// draw the crosshair
 	FCanvasTileItem TileItem( CrosshairDrawPosition, CrosshairTex->Resource, FLinearColor::White);
