Added SpawnCollisionHandlingOverride to the projectile spawner to be consistant with how the BP template is setup. This adds a debug message to the output log when the projectile doesn't spawn.

diff --git a/Templates/TP_FirstPerson/Source/TP_FirstPerson/TP_FirstPersonCharacter.cpp b/Templates/TP_FirstPerson/Source/TP_FirstPerson/TP_FirstPersonCharacter.cpp
--- a/Templates/TP_FirstPerson/Source/TP_FirstPerson/TP_FirstPersonCharacter.cpp
+++ b/Templates/TP_FirstPerson/Source/TP_FirstPerson/TP_FirstPersonCharacter.cpp
@@ -152,8 +152,12 @@ void ATP_FirstPersonCharacter::OnFire()
 				// MuzzleOffset is in camera space, so transform it to world space before offsetting from the character location to find the final muzzle position
 				const FVector SpawnLocation = ((FP_MuzzleLocation != nullptr) ? FP_MuzzleLocation->GetComponentLocation() : GetActorLocation()) + SpawnRotation.RotateVector(GunOffset);
 
+				//Set Spawn Collision Handling Override
+				FActorSpawnParameters ActorSpawnParams;
+				ActorSpawnParams.SpawnCollisionHandlingOverride = ESpawnActorCollisionHandlingMethod::AdjustIfPossibleButDontSpawnIfColliding;
+
 				// spawn the projectile at the muzzle
-				World->SpawnActor<ATP_FirstPersonProjectile>(ProjectileClass, SpawnLocation, SpawnRotation);
+				World->SpawnActor<ATP_FirstPersonProjectile>(ProjectileClass, SpawnLocation, SpawnRotation, ActorSpawnParams);
 			}
 		}
 	}
