Commented out TouchUpdate in TP_FirstPersonCharacter.cpp to make the code template consistent with BP.
The user can no longer rotate without using the right joystick.

diff --git a/Templates/TP_FirstPerson/Source/TP_FirstPerson/TP_FirstPersonCharacter.cpp b/Templates/TP_FirstPerson/Source/TP_FirstPerson/TP_FirstPersonCharacter.cpp
--- a/Templates/TP_FirstPerson/Source/TP_FirstPerson/TP_FirstPersonCharacter.cpp
+++ b/Templates/TP_FirstPerson/Source/TP_FirstPerson/TP_FirstPersonCharacter.cpp
@@ -210,40 +210,43 @@ void ATP_FirstPersonCharacter::EndTouch(const ETouchIndex::Type FingerIndex, con
 	TouchItem.bIsPressed = false;
 }
 
-void ATP_FirstPersonCharacter::TouchUpdate(const ETouchIndex::Type FingerIndex, const FVector Location)
-{
-	if ((TouchItem.bIsPressed == true) && (TouchItem.FingerIndex == FingerIndex))
-	{
-		if (TouchItem.bIsPressed)
-		{
-			if (GetWorld() != nullptr)
-			{
-				UGameViewportClient* ViewportClient = GetWorld()->GetGameViewport();
-				if (ViewportClient != nullptr)
-				{
-					FVector MoveDelta = Location - TouchItem.Location;
-					FVector2D ScreenSize;
-					ViewportClient->GetViewportSize(ScreenSize);
-					FVector2D ScaledDelta = FVector2D(MoveDelta.X, MoveDelta.Y) / ScreenSize;
-					if (FMath::Abs(ScaledDelta.X) >= 4.0 / ScreenSize.X)
-					{
-						TouchItem.bMoved = true;
-						float Value = ScaledDelta.X * BaseTurnRate;
-						AddControllerYawInput(Value);
-					}
-					if (FMath::Abs(ScaledDelta.Y) >= 4.0 / ScreenSize.Y)
-					{
-						TouchItem.bMoved = true;
-						float Value = ScaledDelta.Y * BaseTurnRate;
-						AddControllerPitchInput(Value);
-					}
-					TouchItem.Location = Location;
-				}
-				TouchItem.Location = Location;
-			}
-		}
-	}
-}
+//Commenting this section out to be consistent with FPS BP template.
+//This allows the user to turn without using the right virtual joystick
+
+//void ATP_FirstPersonCharacter::TouchUpdate(const ETouchIndex::Type FingerIndex, const FVector Location)
+//{
+//	if ((TouchItem.bIsPressed == true) && (TouchItem.FingerIndex == FingerIndex))
+//	{
+//		if (TouchItem.bIsPressed)
+//		{
+//			if (GetWorld() != nullptr)
+//			{
+//				UGameViewportClient* ViewportClient = GetWorld()->GetGameViewport();
+//				if (ViewportClient != nullptr)
+//				{
+//					FVector MoveDelta = Location - TouchItem.Location;
+//					FVector2D ScreenSize;
+//					ViewportClient->GetViewportSize(ScreenSize);
+//					FVector2D ScaledDelta = FVector2D(MoveDelta.X, MoveDelta.Y) / ScreenSize;
+//					if (FMath::Abs(ScaledDelta.X) >= 4.0 / ScreenSize.X)
+//					{
+//						TouchItem.bMoved = true;
+//						float Value = ScaledDelta.X * BaseTurnRate;
+//						AddControllerYawInput(Value);
+//					}
+//					if (FMath::Abs(ScaledDelta.Y) >= 4.0 / ScreenSize.Y)
+//					{
+//						TouchItem.bMoved = true;
+//						float Value = ScaledDelta.Y * BaseTurnRate;
+//						AddControllerPitchInput(Value);
+//					}
+//					TouchItem.Location = Location;
+//				}
+//				TouchItem.Location = Location;
+//			}
+//		}
+//	}
+//}
 
 void ATP_FirstPersonCharacter::MoveForward(float Value)
 {
@@ -283,7 +286,9 @@ bool ATP_FirstPersonCharacter::EnableTouchscreenMovement(class UInputComponent*
 		bResult = true;
 		PlayerInputComponent->BindTouch(EInputEvent::IE_Pressed, this, &ATP_FirstPersonCharacter::BeginTouch);
 		PlayerInputComponent->BindTouch(EInputEvent::IE_Released, this, &ATP_FirstPersonCharacter::EndTouch);
-		PlayerInputComponent->BindTouch(EInputEvent::IE_Repeat, this, &ATP_FirstPersonCharacter::TouchUpdate);
+
+		//Commenting this out to be more consistent with FPS BP template.
+		//PlayerInputComponent->BindTouch(EInputEvent::IE_Repeat, this, &ATP_FirstPersonCharacter::TouchUpdate);
 	}
 	return bResult;
 }
