Fixup deprecation warnings in project Templates

diff --git a/Templates/FP_FirstPerson/Source/FP_FirstPerson/FP_FirstPersonCharacter.cpp b/Templates/FP_FirstPerson/Source/FP_FirstPerson/FP_FirstPersonCharacter.cpp
--- a/Templates/FP_FirstPerson/Source/FP_FirstPerson/FP_FirstPersonCharacter.cpp
+++ b/Templates/FP_FirstPerson/Source/FP_FirstPerson/FP_FirstPersonCharacter.cpp
@@ -27,7 +27,7 @@ AFP_FirstPersonCharacter::AFP_FirstPersonCharacter()
 	// Create a CameraComponent	
 	FirstPersonCameraComponent = CreateDefaultSubobject<UCameraComponent>(TEXT("FirstPersonCamera"));
 	FirstPersonCameraComponent->SetupAttachment(GetCapsuleComponent());
-	FirstPersonCameraComponent->RelativeLocation = FVector(0, 0, 64.f); // Position the camera
+	FirstPersonCameraComponent->SetRelativeLocation(FVector(0, 0, 64.f)); // Position the camera
 	FirstPersonCameraComponent->bUsePawnControlRotation = true;
 	
 	// Create a mesh component that will be used when being viewed from a '1st person' view (when controlling this pawn)
@@ -246,7 +246,7 @@ void AFP_FirstPersonCharacter::LookUpAtRate(float Rate)
 FHitResult AFP_FirstPersonCharacter::WeaponTrace(const FVector& StartTrace, const FVector& EndTrace) const
 {
 	// Perform trace to retrieve hit info
-	FCollisionQueryParams TraceParams(SCENE_QUERY_STAT(WeaponTrace), true, Instigator);
+	FCollisionQueryParams TraceParams(SCENE_QUERY_STAT(WeaponTrace), true, GetInstigator());
 	TraceParams.bReturnPhysicalMaterial = true;
 
 	FHitResult Hit(ForceInit);
diff --git a/Templates/TP_2DSideScroller/Source/TP_2DSideScroller/TP_2DSideScrollerCharacter.cpp b/Templates/TP_2DSideScroller/Source/TP_2DSideScroller/TP_2DSideScrollerCharacter.cpp
--- a/Templates/TP_2DSideScroller/Source/TP_2DSideScroller/TP_2DSideScrollerCharacter.cpp
+++ b/Templates/TP_2DSideScroller/Source/TP_2DSideScroller/TP_2DSideScrollerCharacter.cpp
@@ -31,9 +31,9 @@ ATP_2DSideScrollerCharacter::ATP_2DSideScrollerCharacter()
 	CameraBoom->SetupAttachment(RootComponent);
 	CameraBoom->TargetArmLength = 500.0f;
 	CameraBoom->SocketOffset = FVector(0.0f, 0.0f, 75.0f);
-	CameraBoom->bAbsoluteRotation = true;
+	CameraBoom->SetUsingAbsoluteRotation(true);
 	CameraBoom->bDoCollisionTest = false;
-	CameraBoom->RelativeRotation = FRotator(0.0f, -90.0f, 0.0f);
+	CameraBoom->SetRelativeRotation(FRotator(0.0f, -90.0f, 0.0f));
 	
 
 	// Create an orthographic camera (no perspective) and attach it to the boom
@@ -43,7 +43,7 @@ ATP_2DSideScrollerCharacter::ATP_2DSideScrollerCharacter()
 	SideViewCameraComponent->SetupAttachment(CameraBoom, USpringArmComponent::SocketName);
 
 	// Prevent all automatic rotation behavior on the camera, character, and camera component
-	CameraBoom->bAbsoluteRotation = true;
+	CameraBoom->SetUsingAbsoluteRotation(true);
 	SideViewCameraComponent->bUsePawnControlRotation = false;
 	SideViewCameraComponent->bAutoActivate = true;
 	GetCharacterMovement()->bOrientRotationToMovement = false;
diff --git a/Templates/TP_FirstPerson/Source/TP_FirstPerson/TP_FirstPersonCharacter.cpp b/Templates/TP_FirstPerson/Source/TP_FirstPerson/TP_FirstPersonCharacter.cpp
--- a/Templates/TP_FirstPerson/Source/TP_FirstPerson/TP_FirstPersonCharacter.cpp
+++ b/Templates/TP_FirstPerson/Source/TP_FirstPerson/TP_FirstPersonCharacter.cpp
@@ -29,7 +29,7 @@ ATP_FirstPersonCharacter::ATP_FirstPersonCharacter()
 	// Create a CameraComponent	
 	FirstPersonCameraComponent = CreateDefaultSubobject<UCameraComponent>(TEXT("FirstPersonCamera"));
 	FirstPersonCameraComponent->SetupAttachment(GetCapsuleComponent());
-	FirstPersonCameraComponent->RelativeLocation = FVector(-39.56f, 1.75f, 64.f); // Position the camera
+	FirstPersonCameraComponent->SetRelativeLocation(FVector(-39.56f, 1.75f, 64.f)); // Position the camera
 	FirstPersonCameraComponent->bUsePawnControlRotation = true;
 
 	// Create a mesh component that will be used when being viewed from a '1st person' view (when controlling this pawn)
@@ -38,8 +38,8 @@ ATP_FirstPersonCharacter::ATP_FirstPersonCharacter()
 	Mesh1P->SetupAttachment(FirstPersonCameraComponent);
 	Mesh1P->bCastDynamicShadow = false;
 	Mesh1P->CastShadow = false;
-	Mesh1P->RelativeRotation = FRotator(1.9f, -19.19f, 5.2f);
-	Mesh1P->RelativeLocation = FVector(-0.5f, -4.4f, -155.7f);
+	Mesh1P->SetRelativeRotation(FRotator(1.9f, -19.19f, 5.2f));
+	Mesh1P->SetRelativeLocation(FVector(-0.5f, -4.4f, -155.7f));
 
 	// Create a gun mesh component
 	FP_Gun = CreateDefaultSubobject<USkeletalMeshComponent>(TEXT("FP_Gun"));
diff --git a/Templates/TP_Rolling/Source/TP_Rolling/TP_RollingBall.cpp b/Templates/TP_Rolling/Source/TP_Rolling/TP_RollingBall.cpp
--- a/Templates/TP_Rolling/Source/TP_Rolling/TP_RollingBall.cpp
+++ b/Templates/TP_Rolling/Source/TP_Rolling/TP_RollingBall.cpp
@@ -29,8 +29,8 @@ ATP_RollingBall::ATP_RollingBall()
 	SpringArm = CreateDefaultSubobject<USpringArmComponent>(TEXT("SpringArm0"));
 	SpringArm->SetupAttachment(RootComponent);
 	SpringArm->bDoCollisionTest = false;
-	SpringArm->bAbsoluteRotation = true; // Rotation of the ball should not affect rotation of boom
-	SpringArm->RelativeRotation = FRotator(-45.f, 0.f, 0.f);
+	SpringArm->SetUsingAbsoluteRotation(true); // Rotation of the ball should not affect rotation of boom
+	SpringArm->SetRelativeRotation(FRotator(-45.f, 0.f, 0.f));
 	SpringArm->TargetArmLength = 1200.f;
 	SpringArm->bEnableCameraLag = false;
 	SpringArm->CameraLagSpeed = 3.f;
diff --git a/Templates/TP_SideScroller/Source/TP_SideScroller/TP_SideScrollerCharacter.cpp b/Templates/TP_SideScroller/Source/TP_SideScroller/TP_SideScrollerCharacter.cpp
--- a/Templates/TP_SideScroller/Source/TP_SideScroller/TP_SideScrollerCharacter.cpp
+++ b/Templates/TP_SideScroller/Source/TP_SideScroller/TP_SideScrollerCharacter.cpp
@@ -20,11 +20,11 @@ ATP_SideScrollerCharacter::ATP_SideScrollerCharacter()
 	// Create a camera boom attached to the root (capsule)
 	CameraBoom = CreateDefaultSubobject<USpringArmComponent>(TEXT("CameraBoom"));
 	CameraBoom->SetupAttachment(RootComponent);
-	CameraBoom->bAbsoluteRotation = true; // Rotation of the character should not affect rotation of boom
+	CameraBoom->SetUsingAbsoluteRotation(true); // Rotation of the character should not affect rotation of boom
 	CameraBoom->bDoCollisionTest = false;
 	CameraBoom->TargetArmLength = 500.f;
 	CameraBoom->SocketOffset = FVector(0.f,0.f,75.f);
-	CameraBoom->RelativeRotation = FRotator(0.f,180.f,0.f);
+	CameraBoom->SetRelativeRotation(FRotator(0.f,180.f,0.f));
 
 	// Create a camera and attach to boom
 	SideViewCameraComponent = CreateDefaultSubobject<UCameraComponent>(TEXT("SideViewCamera"));
diff --git a/Templates/TP_TopDown/Source/TP_TopDown/TP_TopDownCharacter.cpp b/Templates/TP_TopDown/Source/TP_TopDown/TP_TopDownCharacter.cpp
--- a/Templates/TP_TopDown/Source/TP_TopDown/TP_TopDownCharacter.cpp
+++ b/Templates/TP_TopDown/Source/TP_TopDown/TP_TopDownCharacter.cpp
@@ -31,9 +31,9 @@ ATP_TopDownCharacter::ATP_TopDownCharacter()
 	// Create a camera boom...
 	CameraBoom = CreateDefaultSubobject<USpringArmComponent>(TEXT("CameraBoom"));
 	CameraBoom->SetupAttachment(RootComponent);
-	CameraBoom->bAbsoluteRotation = true; // Don't want arm to rotate when character does
+	CameraBoom->SetUsingAbsoluteRotation(true); // Don't want arm to rotate when character does
 	CameraBoom->TargetArmLength = 800.f;
-	CameraBoom->RelativeRotation = FRotator(-60.f, 0.f, 0.f);
+	CameraBoom->SetRelativeRotation(FRotator(-60.f, 0.f, 0.f));
 	CameraBoom->bDoCollisionTest = false; // Don't want to pull camera in when it collides with level
 
 	// Create a camera...
diff --git a/Templates/TP_TwinStick/Source/TP_TwinStick/TP_TwinStickPawn.cpp b/Templates/TP_TwinStick/Source/TP_TwinStick/TP_TwinStickPawn.cpp
--- a/Templates/TP_TwinStick/Source/TP_TwinStick/TP_TwinStickPawn.cpp
+++ b/Templates/TP_TwinStick/Source/TP_TwinStick/TP_TwinStickPawn.cpp
@@ -34,9 +34,9 @@ ATP_TwinStickPawn::ATP_TwinStickPawn()
 	// Create a camera boom...
 	CameraBoom = CreateDefaultSubobject<USpringArmComponent>(TEXT("CameraBoom"));
 	CameraBoom->SetupAttachment(RootComponent);
-	CameraBoom->bAbsoluteRotation = true; // Don't want arm to rotate when ship does
+	CameraBoom->SetUsingAbsoluteRotation(true); // Don't want arm to rotate when ship does
 	CameraBoom->TargetArmLength = 1200.f;
-	CameraBoom->RelativeRotation = FRotator(-80.f, 0.f, 0.f);
+	CameraBoom->SetRelativeRotation(FRotator(-80.f, 0.f, 0.f));
 	CameraBoom->bDoCollisionTest = false; // Don't want to pull camera in when it collides with level
 
 	// Create a camera...
diff --git a/Templates/TP_Vehicle/Source/TP_Vehicle/TP_VehiclePawn.cpp b/Templates/TP_Vehicle/Source/TP_Vehicle/TP_VehiclePawn.cpp
--- a/Templates/TP_Vehicle/Source/TP_Vehicle/TP_VehiclePawn.cpp
+++ b/Templates/TP_Vehicle/Source/TP_Vehicle/TP_VehiclePawn.cpp
@@ -214,10 +214,10 @@ void ATP_VehiclePawn::Tick(float Delta)
 	{
 		if ( (InputComponent) && (bInCarCameraActive == true ))
 		{
-			FRotator HeadRotation = InternalCamera->RelativeRotation;
+			FRotator HeadRotation = InternalCamera->GetRelativeRotation();
 			HeadRotation.Pitch += InputComponent->GetAxisValue(LookUpBinding);
 			HeadRotation.Yaw += InputComponent->GetAxisValue(LookRightBinding);
-			InternalCamera->RelativeRotation = HeadRotation;
+			InternalCamera->SetRelativeRotation(HeadRotation);
 		}
 	}
 }
diff --git a/Templates/TP_VehicleAdv/Source/TP_VehicleAdv/TP_VehicleAdvPawn.cpp b/Templates/TP_VehicleAdv/Source/TP_VehicleAdv/TP_VehicleAdvPawn.cpp
--- a/Templates/TP_VehicleAdv/Source/TP_VehicleAdv/TP_VehicleAdvPawn.cpp
+++ b/Templates/TP_VehicleAdv/Source/TP_VehicleAdv/TP_VehicleAdvPawn.cpp
@@ -272,10 +272,10 @@ void ATP_VehicleAdvPawn::Tick(float Delta)
 	{
 		if ( (InputComponent) && (bInCarCameraActive == true ))
 		{
-			FRotator HeadRotation = InternalCamera->RelativeRotation;
+			FRotator HeadRotation = InternalCamera->GetRelativeRotation();
 			HeadRotation.Pitch += InputComponent->GetAxisValue(LookUpBinding);
 			HeadRotation.Yaw += InputComponent->GetAxisValue(LookRightBinding);
-			InternalCamera->RelativeRotation = HeadRotation;
+			InternalCamera->SetRelativeRotation(HeadRotation);
 		}
 	}	
 
