Prepare for the deprecation of direct access to the Resource member of UTexture in favor of the GetResource() accessor.

diff --git a/Templates/TP_FirstPerson/Source/TP_FirstPerson/TP_FirstPersonHUD.cpp b/Templates/TP_FirstPerson/Source/TP_FirstPerson/TP_FirstPersonHUD.cpp
--- a/Templates/TP_FirstPerson/Source/TP_FirstPerson/TP_FirstPersonHUD.cpp
+++ b/Templates/TP_FirstPerson/Source/TP_FirstPerson/TP_FirstPersonHUD.cpp
@@ -29,7 +29,7 @@ void ATP_FirstPersonHUD::DrawHUD()
 										   (Center.Y + 20.0f));
 
 	// draw the crosshair
-	FCanvasTileItem TileItem( CrosshairDrawPosition, CrosshairTex->Resource, FLinearColor::White);
+	FCanvasTileItem TileItem( CrosshairDrawPosition, CrosshairTex->GetResource(), FLinearColor::White);
 	TileItem.BlendMode = SE_BLEND_Translucent;
 	Canvas->DrawItem( TileItem );
 }
