Templates: TopDown, FirstPerson, ThirdPerson C++ comments and logic update. Added blendspace for WalkRun using analyze to reduce foot sliding. General cleanup. Saving maps with proper cam. Assitionally setting Game Mode in map overwrite.

diff --git a/Templates/TP_FirstPerson/Source/TP_FirstPerson/TP_FirstPersonCharacter.cpp b/Templates/TP_FirstPerson/Source/TP_FirstPerson/TP_FirstPersonCharacter.cpp
--- a/Templates/TP_FirstPerson/Source/TP_FirstPerson/TP_FirstPersonCharacter.cpp
+++ b/Templates/TP_FirstPerson/Source/TP_FirstPerson/TP_FirstPersonCharacter.cpp
@@ -77,7 +77,7 @@ void ATP_FirstPersonCharacter::SetupPlayerInputComponent(class UInputComponent*
 void ATP_FirstPersonCharacter::OnPrimaryAction()
 {
 	// Trigger the OnItemUsed Event
-	OnItemUsed.Broadcast();
+	OnUseItem.Broadcast();
 }
 
 void ATP_FirstPersonCharacter::BeginTouch(const ETouchIndex::Type FingerIndex, const FVector Location)
@@ -105,44 +105,6 @@ void ATP_FirstPersonCharacter::EndTouch(const ETouchIndex::Type FingerIndex, con
 	TouchItem.bIsPressed = false;
 }
 
-//Commenting this section out to be consistent with FPS BP template.
-//This allows the user to turn without using the right virtual joystick
-
-//void ATP_FirstPersonCharacter::TouchUpdate(const ETouchIndex::Type FingerIndex, const FVector Location)
-//{
-//	if ((TouchItem.bIsPressed == true) && (TouchItem.FingerIndex == FingerIndex))
-//	{
-//		if (TouchItem.bIsPressed)
-//		{
-//			if (GetWorld() != nullptr)
-//			{
-//				UGameViewportClient* ViewportClient = GetWorld()->GetGameViewport();
-//				if (ViewportClient != nullptr)
-//				{
-//					FVector MoveDelta = Location - TouchItem.Location;
-//					FVector2D ScreenSize;
-//					ViewportClient->GetViewportSize(ScreenSize);
-//					FVector2D ScaledDelta = FVector2D(MoveDelta.X, MoveDelta.Y) / ScreenSize;
-//					if (FMath::Abs(ScaledDelta.X) >= 4.0 / ScreenSize.X)
-//					{
-//						TouchItem.bMoved = true;
-//						float Value = ScaledDelta.X * TurnRateGamepad;
-//						AddControllerYawInput(Value);
-//					}
-//					if (FMath::Abs(ScaledDelta.Y) >= 4.0 / ScreenSize.Y)
-//					{
-//						TouchItem.bMoved = true;
-//						float Value = ScaledDelta.Y * TurnRateGamepad;
-//						AddControllerPitchInput(Value);
-//					}
-//					TouchItem.Location = Location;
-//				}
-//				TouchItem.Location = Location;
-//			}
-//		}
-//	}
-//}
-
 void ATP_FirstPersonCharacter::MoveForward(float Value)
 {
 	if (Value != 0.0f)
@@ -180,8 +142,6 @@ bool ATP_FirstPersonCharacter::EnableTouchscreenMovement(class UInputComponent*
 		PlayerInputComponent->BindTouch(EInputEvent::IE_Pressed, this, &ATP_FirstPersonCharacter::BeginTouch);
 		PlayerInputComponent->BindTouch(EInputEvent::IE_Released, this, &ATP_FirstPersonCharacter::EndTouch);
 
-		//Commenting this out to be more consistent with FPS BP template.
-		//PlayerInputComponent->BindTouch(EInputEvent::IE_Repeat, this, &ATP_FirstPersonCharacter::TouchUpdate);
 		return true;
 	}
 	
diff --git a/Templates/TP_FirstPerson/Source/TP_FirstPerson/TP_FirstPersonCharacter.h b/Templates/TP_FirstPerson/Source/TP_FirstPerson/TP_FirstPersonCharacter.h
--- a/Templates/TP_FirstPerson/Source/TP_FirstPerson/TP_FirstPersonCharacter.h
+++ b/Templates/TP_FirstPerson/Source/TP_FirstPerson/TP_FirstPersonCharacter.h
@@ -13,7 +13,9 @@ class UCameraComponent;
 class UAnimMontage;
 class USoundBase;
 
-DECLARE_DYNAMIC_MULTICAST_DELEGATE(FOnItemUsed);
+// Declaration of the delegate that will be called when the Primary Action is triggered
+// It is declared as dynamic so it can be accessed also in Blueprints
+DECLARE_DYNAMIC_MULTICAST_DELEGATE(FOnUseItem);
 
 UCLASS(config=Game)
 class ATP_FirstPersonCharacter : public ACharacter
@@ -39,8 +41,9 @@ public:
 	UPROPERTY(VisibleAnywhere, BlueprintReadOnly, Category=Camera)
 	float TurnRateGamepad;
 
+	/** Delegate to whom anyone can subscribe to receive this event */
 	UPROPERTY(BlueprintAssignable, Category = "Interaction")
-	FOnItemUsed OnItemUsed;
+	FOnUseItem OnUseItem;
 protected:
 	
 	/** Fires a projectile. */
diff --git a/Templates/TP_FirstPerson/Source/TP_FirstPerson/TP_PickUpComponent.cpp b/Templates/TP_FirstPerson/Source/TP_FirstPerson/TP_PickUpComponent.cpp
new file mode 100644
--- /dev/null
+++ b/Templates/TP_FirstPerson/Source/TP_FirstPerson/TP_PickUpComponent.cpp
@@ -0,0 +1,31 @@
+// Copyright Epic Games, Inc. All Rights Reserved.
+
+#include "TP_PickUpComponent.h"
+
+UTP_PickUpComponent::UTP_PickUpComponent()
+{
+	// Setup the Sphere Collision
+	SphereRadius = 32.f;
+}
+
+void UTP_PickUpComponent::BeginPlay()
+{
+	Super::BeginPlay();
+
+	// Register our Overlap Event
+	OnComponentBeginOverlap.AddDynamic(this, &UTP_PickUpComponent::OnSphereBeginOverlap);
+}
+
+void UTP_PickUpComponent::OnSphereBeginOverlap(UPrimitiveComponent* OverlappedComponent, AActor* OtherActor, UPrimitiveComponent* OtherComp, int32 OtherBodyIndex, bool bFromSweep, const FHitResult& SweepResult)
+{
+	// Checking if it is a First Person Character overlapping
+	ATP_FirstPersonCharacter* Character = Cast<ATP_FirstPersonCharacter>(OtherActor);
+	if(Character != nullptr)
+	{
+		// Notify that the actor is being picked up
+		OnPickUp.Broadcast(Character);
+
+		// Unregister from the Overlap Event so it is no longer triggered
+		OnComponentBeginOverlap.RemoveAll(this);
+	}
+}
diff --git a/Templates/TP_FirstPerson/Source/TP_FirstPerson/TP_PickUpComponent.h b/Templates/TP_FirstPerson/Source/TP_FirstPerson/TP_PickUpComponent.h
new file mode 100644
--- /dev/null
+++ b/Templates/TP_FirstPerson/Source/TP_FirstPerson/TP_PickUpComponent.h
@@ -0,0 +1,34 @@
+// Copyright Epic Games, Inc. All Rights Reserved.
+
+#pragma once
+
+#include "CoreMinimal.h"
+#include "Components/SphereComponent.h"
+#include "TP_FirstPersonCharacter.h"
+#include "TP_PickUpComponent.generated.h"
+
+// Declaration of the delegate that will be called when someone picks this up
+// The character picking this up is the parameter sent with the notification
+DECLARE_DYNAMIC_MULTICAST_DELEGATE_OneParam(FOnPickUp, ATP_FirstPersonCharacter*, PickUpCharacter);
+
+UCLASS(Blueprintable, BlueprintType, ClassGroup = (Custom), meta = (BlueprintSpawnableComponent))
+class TP_FIRSTPERSON_API UTP_PickUpComponent : public USphereComponent
+{
+	GENERATED_BODY()
+
+public:
+	
+	/** Delegate to whom anyone can subscribe to receive this event */
+	UPROPERTY(BlueprintAssignable, Category = "Interaction")
+	FOnPickUp OnPickUp;
+
+	UTP_PickUpComponent();
+protected:
+
+	/** Called when the game starts */
+	virtual void BeginPlay() override;
+
+	/** Code for when something overlaps this component */
+	UFUNCTION()
+	void OnSphereBeginOverlap(UPrimitiveComponent* OverlappedComponent, AActor* OtherActor, UPrimitiveComponent* OtherComp, int32 OtherBodyIndex, bool bFromSweep, const FHitResult& SweepResult);
+};
\ No newline at end of file
diff --git a/Templates/TP_FirstPerson/Source/TP_FirstPerson/TP_WeaponComponent.cpp b/Templates/TP_FirstPerson/Source/TP_FirstPerson/TP_WeaponComponent.cpp
--- a/Templates/TP_FirstPerson/Source/TP_FirstPerson/TP_WeaponComponent.cpp
+++ b/Templates/TP_FirstPerson/Source/TP_FirstPerson/TP_WeaponComponent.cpp
@@ -11,7 +11,6 @@
 // Sets default values for this component's properties
 UTP_WeaponComponent::UTP_WeaponComponent()
 {
-	SphereRadius = 32.f;
 	// Default offset from the character location for projectiles to spawn
 	MuzzleOffset = FVector(100.0f, 0.0f, 10.0f);
 }
@@ -62,29 +61,26 @@ void UTP_WeaponComponent::Fire()
 	}
 }
 
-// Called when the game starts
-void UTP_WeaponComponent::BeginPlay()
+void UTP_WeaponComponent::EndPlay(const EEndPlayReason::Type EndPlayReason)
 {
-	Super::BeginPlay();
-
-	// Register our Overlap Event
-	OnComponentBeginOverlap.AddDynamic(this, &UTP_WeaponComponent::OnSphereBeginOverlap);
+	if(Character != nullptr)
+	{
+		// Unregister from the OnUseItem Event
+		Character->OnUseItem.RemoveDynamic(this, &UTP_WeaponComponent::Fire);
+	}
 }
 
-void UTP_WeaponComponent::OnSphereBeginOverlap(UPrimitiveComponent* OverlappedComponent, AActor* OtherActor, UPrimitiveComponent* OtherComp, int32 OtherBodyIndex, bool bFromSweep, const FHitResult& SweepResult)
+void UTP_WeaponComponent::AttachWeapon(ATP_FirstPersonCharacter* TargetCharacter)
 {
-	// Checking if it is a First Person Character overlapping
-	Character = Cast<ATP_FirstPersonCharacter>(OtherActor);
+	Character = TargetCharacter;
 	if(Character != nullptr)
 	{
 		// Attach the weapon to the First Person Character
 		FAttachmentTransformRules AttachmentRules(EAttachmentRule::SnapToTarget, true);
 		GetOwner()->AttachToComponent(Character->GetMesh1P(),AttachmentRules, FName(TEXT("GripPoint")));
 
-		// Unregister from the Overlap Event so it is no longer triggered
-		OnComponentBeginOverlap.RemoveAll(this);
 		// Register so that Fire is called every time the character tries to use the item being held
-		Character->OnItemUsed.AddDynamic(this, &UTP_WeaponComponent::Fire);
+		Character->OnUseItem.AddDynamic(this, &UTP_WeaponComponent::Fire);
 	}
 }
 
diff --git a/Templates/TP_FirstPerson/Source/TP_FirstPerson/TP_WeaponComponent.h b/Templates/TP_FirstPerson/Source/TP_FirstPerson/TP_WeaponComponent.h
--- a/Templates/TP_FirstPerson/Source/TP_FirstPerson/TP_WeaponComponent.h
+++ b/Templates/TP_FirstPerson/Source/TP_FirstPerson/TP_WeaponComponent.h
@@ -4,20 +4,14 @@
 
 #include "CoreMinimal.h"
 #include "Components/ActorComponent.h"
-#include "Components/SphereComponent.h"
 #include "TP_WeaponComponent.generated.h"
 
 class ATP_FirstPersonCharacter;
-class USkeletalMeshComponent;
-class USceneComponent;
 
 UCLASS(Blueprintable, BlueprintType, ClassGroup=(Custom), meta=(BlueprintSpawnableComponent) )
-class TP_FIRSTPERSON_API UTP_WeaponComponent : public USphereComponent
+class TP_FIRSTPERSON_API UTP_WeaponComponent : public UActorComponent
 {
 	GENERATED_BODY()
-	/** Sphere collision to do Pick up */
-	UPROPERTY(VisibleAnywhere, BlueprintReadOnly, Category = Mesh, meta = (AllowPrivateAccess = "true"))
-	USphereComponent* SphereCollision;
 
 public:
 	/** Projectile class to spawn */
@@ -36,20 +30,22 @@ public:
 	UPROPERTY(EditAnywhere, BlueprintReadWrite, Category=Gameplay)
 	FVector MuzzleOffset;
 
-	// Sets default values for this component's properties
+	/** Sets default values for this component's properties */
 	UTP_WeaponComponent();
 
-	// Make the weapon Fire a Projectile
+	/** Attaches the actor to a FirstPersonCharacter */
+	UFUNCTION(BlueprintCallable, Category="Weapon")
+	void AttachWeapon(ATP_FirstPersonCharacter* TargetCharacter);
+
+	/** Make the weapon Fire a Projectile */
 	UFUNCTION(BlueprintCallable, Category="Weapon")
 	void Fire();
 
 protected:
-	// Called when the game starts
-	virtual void BeginPlay() override;
-
-	// Code for when something overlaps this component
+	/** Ends gameplay for this component. */
 	UFUNCTION()
-	void OnSphereBeginOverlap(UPrimitiveComponent* OverlappedComponent, AActor* OtherActor, UPrimitiveComponent* OtherComp, int32 OtherBodyIndex, bool bFromSweep, const FHitResult& SweepResult);
+	virtual void EndPlay(const EEndPlayReason::Type EndPlayReason) override;
+		
 
 private:
 	/** The Character holding this weapon*/
