[Templates] Fixing exposure issue, switching to SM6, adding Enhanced Input for C++ versions, adding Asset Guideline for plugin for chaos vehicle (feature pack relevant) and several fixes. (Known issue: shooting in 1st person not working)

#content only

diff --git a/Templates/TP_FirstPerson/Config/DefaultInput.ini b/Templates/TP_FirstPerson/Config/DefaultInput.ini
--- a/Templates/TP_FirstPerson/Config/DefaultInput.ini
+++ b/Templates/TP_FirstPerson/Config/DefaultInput.ini
@@ -90,6 +90,9 @@ bEnableMouseSmoothing=True
 bEnableFOVScaling=True
 bCaptureMouseOnLaunch=True
 bEnableLegacyInputScales=True
+bEnableMotionControls=True
+bFilterInputByPlatformUser=False
+bShouldFlushPressedKeysOnViewportFocusLost=True
 bAlwaysShowTouchInterface=False
 bShowConsoleOnFourFingerTap=True
 bEnableGestureRecognizer=False
@@ -98,24 +101,6 @@ DefaultViewportMouseCaptureMode=CapturePermanently_IncludingInitialMouseDown
 DefaultViewportMouseLockMode=LockOnCapture
 FOVScale=0.011110
 DoubleClickTime=0.200000
-+ActionMappings=(ActionName="Jump",bShift=False,bCtrl=False,bAlt=False,bCmd=False,Key=SpaceBar)
-+ActionMappings=(ActionName="Jump",bShift=False,bCtrl=False,bAlt=False,bCmd=False,Key=Gamepad_FaceButton_Bottom)
-+ActionMappings=(ActionName="PrimaryAction",bShift=False,bCtrl=False,bAlt=False,bCmd=False,Key=LeftMouseButton)
-+ActionMappings=(ActionName="PrimaryAction",bShift=False,bCtrl=False,bAlt=False,bCmd=False,Key=Gamepad_RightTrigger)
-+AxisMappings=(AxisName="Move Forward / Backward",Scale=1.000000,Key=W)
-+AxisMappings=(AxisName="Move Forward / Backward",Scale=-1.000000,Key=S)
-+AxisMappings=(AxisName="Move Forward / Backward",Scale=1.000000,Key=Up)
-+AxisMappings=(AxisName="Move Forward / Backward",Scale=-1.000000,Key=Down)
-+AxisMappings=(AxisName="Move Forward / Backward",Scale=1.000000,Key=Gamepad_LeftY)
-+AxisMappings=(AxisName="Move Right / Left",Scale=-1.000000,Key=A)
-+AxisMappings=(AxisName="Move Right / Left",Scale=1.000000,Key=D)
-+AxisMappings=(AxisName="Move Right / Left",Scale=1.000000,Key=Gamepad_LeftX)
-+AxisMappings=(AxisName="Turn Right / Left Gamepad",Scale=1.000000,Key=Gamepad_RightX)
-+AxisMappings=(AxisName="Turn Right / Left Gamepad",Scale=-1.000000,Key=Left)
-+AxisMappings=(AxisName="Turn Right / Left Gamepad",Scale=1.000000,Key=Right)
-+AxisMappings=(AxisName="Turn Right / Left Mouse",Scale=1.000000,Key=MouseX)
-+AxisMappings=(AxisName="Look Up / Down Gamepad",Scale=1.000000,Key=Gamepad_RightY)
-+AxisMappings=(AxisName="Look Up / Down Mouse",Scale=-1.000000,Key=MouseY)
 DefaultPlayerInputClass=/Script/EnhancedInput.EnhancedPlayerInput
 DefaultInputComponentClass=/Script/EnhancedInput.EnhancedInputComponent
 DefaultTouchInterface=/Engine/MobileResources/HUD/DefaultVirtualJoysticks.DefaultVirtualJoysticks
diff --git a/Templates/TP_FirstPerson/Source/TP_FirstPerson/TP_FirstPerson.Build.cs b/Templates/TP_FirstPerson/Source/TP_FirstPerson/TP_FirstPerson.Build.cs
--- a/Templates/TP_FirstPerson/Source/TP_FirstPerson/TP_FirstPerson.Build.cs
+++ b/Templates/TP_FirstPerson/Source/TP_FirstPerson/TP_FirstPerson.Build.cs
@@ -8,6 +8,6 @@ public class TP_FirstPerson : ModuleRules
 	{
 		PCHUsage = PCHUsageMode.UseExplicitOrSharedPCHs;
 
-		PublicDependencyModuleNames.AddRange(new string[] { "Core", "CoreUObject", "Engine", "InputCore", "HeadMountedDisplay" });
+		PublicDependencyModuleNames.AddRange(new string[] { "Core", "CoreUObject", "Engine", "InputCore", "HeadMountedDisplay", "EnhancedInput" });
 	}
 }
diff --git a/Templates/TP_FirstPerson/Source/TP_FirstPerson/TP_FirstPersonCharacter.cpp b/Templates/TP_FirstPerson/Source/TP_FirstPerson/TP_FirstPersonCharacter.cpp
--- a/Templates/TP_FirstPerson/Source/TP_FirstPerson/TP_FirstPersonCharacter.cpp
+++ b/Templates/TP_FirstPerson/Source/TP_FirstPerson/TP_FirstPersonCharacter.cpp
@@ -5,8 +5,8 @@
 #include "Animation/AnimInstance.h"
 #include "Camera/CameraComponent.h"
 #include "Components/CapsuleComponent.h"
-#include "Components/InputComponent.h"
-#include "GameFramework/InputSettings.h"
+#include "EnhancedInputComponent.h"
+#include "EnhancedInputSubsystems.h"
 
 
 //////////////////////////////////////////////////////////////////////////
@@ -14,16 +14,16 @@
 
 ATP_FirstPersonCharacter::ATP_FirstPersonCharacter()
 {
+	// Character doesnt have a rifle at start
+	bHasRifle = false;
+	
 	// Set size for collision capsule
 	GetCapsuleComponent()->InitCapsuleSize(55.f, 96.0f);
-
-	// set our turn rates for input
-	TurnRateGamepad = 45.f;
-
+		
 	// Create a CameraComponent	
 	FirstPersonCameraComponent = CreateDefaultSubobject<UCameraComponent>(TEXT("FirstPersonCamera"));
 	FirstPersonCameraComponent->SetupAttachment(GetCapsuleComponent());
-	FirstPersonCameraComponent->SetRelativeLocation(FVector(-39.56f, 1.75f, 64.f)); // Position the camera
+	FirstPersonCameraComponent->SetRelativeLocation(FVector(-10.f, 0.f, 60.f)); // Position the camera
 	FirstPersonCameraComponent->bUsePawnControlRotation = true;
 
 	// Create a mesh component that will be used when being viewed from a '1st person' view (when controlling this pawn)
@@ -32,8 +32,8 @@ ATP_FirstPersonCharacter::ATP_FirstPersonCharacter()
 	Mesh1P->SetupAttachment(FirstPersonCameraComponent);
 	Mesh1P->bCastDynamicShadow = false;
 	Mesh1P->CastShadow = false;
-	Mesh1P->SetRelativeRotation(FRotator(1.9f, -19.19f, 5.2f));
-	Mesh1P->SetRelativeLocation(FVector(-0.5f, -4.4f, -155.7f));
+	//Mesh1P->SetRelativeRotation(FRotator(0.9f, -19.19f, 5.2f));
+	Mesh1P->SetRelativeLocation(FVector(-30.f, 0.f, -150.f));
 
 }
 
@@ -48,102 +48,80 @@ void ATP_FirstPersonCharacter::BeginPlay()
 
 void ATP_FirstPersonCharacter::SetupPlayerInputComponent(class UInputComponent* PlayerInputComponent)
 {
-	// Set up gameplay key bindings
-	check(PlayerInputComponent);
-
-	// Bind jump events
-	PlayerInputComponent->BindAction("Jump", IE_Pressed, this, &ACharacter::Jump);
-	PlayerInputComponent->BindAction("Jump", IE_Released, this, &ACharacter::StopJumping);
-
-	// Bind fire event
-	PlayerInputComponent->BindAction("PrimaryAction", IE_Pressed, this, &ATP_FirstPersonCharacter::OnPrimaryAction);
-
-	// Enable touchscreen input
-	EnableTouchscreenMovement(PlayerInputComponent);
-
-	// Bind movement events
-	PlayerInputComponent->BindAxis("Move Forward / Backward", this, &ATP_FirstPersonCharacter::MoveForward);
-	PlayerInputComponent->BindAxis("Move Right / Left", this, &ATP_FirstPersonCharacter::MoveRight);
-
-	// We have 2 versions of the rotation bindings to handle different kinds of devices differently
-	// "Mouse" versions handle devices that provide an absolute delta, such as a mouse.
-	// "Gamepad" versions are for devices that we choose to treat as a rate of change, such as an analog joystick
-	PlayerInputComponent->BindAxis("Turn Right / Left Mouse", this, &APawn::AddControllerYawInput);
-	PlayerInputComponent->BindAxis("Look Up / Down Mouse", this, &APawn::AddControllerPitchInput);
-	PlayerInputComponent->BindAxis("Turn Right / Left Gamepad", this, &ATP_FirstPersonCharacter::TurnAtRate);
-	PlayerInputComponent->BindAxis("Look Up / Down Gamepad", this, &ATP_FirstPersonCharacter::LookUpAtRate);
-}
+	// Set up action bindings
+	if (UEnhancedInputComponent* EnhancedInputComponent = CastChecked<UEnhancedInputComponent>(PlayerInputComponent))
+	{
+		//Jumping
+		EnhancedInputComponent->BindAction(JumpAction, ETriggerEvent::Triggered, this, &ACharacter::Jump);
+		EnhancedInputComponent->BindAction(JumpAction, ETriggerEvent::Completed, this, &ACharacter::StopJumping);
 
-void ATP_FirstPersonCharacter::OnPrimaryAction()
-{
-	// Trigger the OnItemUsed Event
-	OnUseItem.Broadcast();
-}
+		//Moving
+		EnhancedInputComponent->BindAction(MoveAction, ETriggerEvent::Triggered, this, &ATP_FirstPersonCharacter::Move);
 
-void ATP_FirstPersonCharacter::BeginTouch(const ETouchIndex::Type FingerIndex, const FVector Location)
-{
-	if (TouchItem.bIsPressed == true)
-	{
-		return;
-	}
-	if ((FingerIndex == TouchItem.FingerIndex) && (TouchItem.bMoved == false))
-	{
-		OnPrimaryAction();
+		//Looking
+		EnhancedInputComponent->BindAction(LookAction, ETriggerEvent::Triggered, this, &ATP_FirstPersonCharacter::Look);
 	}
-	TouchItem.bIsPressed = true;
-	TouchItem.FingerIndex = FingerIndex;
-	TouchItem.Location = Location;
-	TouchItem.bMoved = false;
 }
 
-void ATP_FirstPersonCharacter::EndTouch(const ETouchIndex::Type FingerIndex, const FVector Location)
+void ATP_FirstPersonCharacter::PossessedBy(AController* NewController)
 {
-	if (TouchItem.bIsPressed == false)
+	Super::PossessedBy(NewController);
+
+	//Add a mapping context on Possessed
+	if (APlayerController* PlayerController = Cast<APlayerController>(NewController))
 	{
-		return;
+		if (UEnhancedInputLocalPlayerSubsystem* Subsystem = ULocalPlayer::GetSubsystem<UEnhancedInputLocalPlayerSubsystem>(PlayerController->GetLocalPlayer()))
+		{
+			Subsystem->AddMappingContext(DefaultMappingContext, 0);
+		}
 	}
-	TouchItem.bIsPressed = false;
 }
 
-void ATP_FirstPersonCharacter::MoveForward(float Value)
+void ATP_FirstPersonCharacter::UnPossessed()
 {
-	if (Value != 0.0f)
+	Super::UnPossessed();
+
+	if (APlayerController* PlayerController = Cast<APlayerController>(GetController()))
 	{
-		// add movement in that direction
-		AddMovementInput(GetActorForwardVector(), Value);
+		if (UEnhancedInputLocalPlayerSubsystem* Subsystem = ULocalPlayer::GetSubsystem<UEnhancedInputLocalPlayerSubsystem>(PlayerController->GetLocalPlayer()))
+		{
+			Subsystem->RemoveMappingContext(DefaultMappingContext);
+		}
 	}
 }
 
-void ATP_FirstPersonCharacter::MoveRight(float Value)
+void ATP_FirstPersonCharacter::Move(const FInputActionValue& Value)
 {
-	if (Value != 0.0f)
+	// input is a Vector2D
+	FVector2D MovementVector = Value.Get<FVector2D>();
+
+	if (Controller != nullptr)
 	{
-		// add movement in that direction
-		AddMovementInput(GetActorRightVector(), Value);
+		// add movement 
+		AddMovementInput(GetActorForwardVector(), MovementVector.Y);
+		AddMovementInput(GetActorRightVector(), MovementVector.X);
 	}
 }
 
-void ATP_FirstPersonCharacter::TurnAtRate(float Rate)
+void ATP_FirstPersonCharacter::Look(const FInputActionValue& Value)
 {
-	// calculate delta for this frame from the rate information
-	AddControllerYawInput(Rate * TurnRateGamepad * GetWorld()->GetDeltaSeconds());
+	// input is a Vector2D
+	FVector2D LookAxisVector = Value.Get<FVector2D>();
+
+	if (Controller != nullptr)
+	{
+		// add yaw and pitch input to controller
+		AddControllerYawInput(LookAxisVector.X);
+		AddControllerPitchInput(LookAxisVector.Y);
+	}
 }
 
-void ATP_FirstPersonCharacter::LookUpAtRate(float Rate)
+void ATP_FirstPersonCharacter::SetHasRifle(bool bNewHasRifle)
 {
-	// calculate delta for this frame from the rate information
-	AddControllerPitchInput(Rate * TurnRateGamepad * GetWorld()->GetDeltaSeconds());
+	bHasRifle = bNewHasRifle;
 }
 
-bool ATP_FirstPersonCharacter::EnableTouchscreenMovement(class UInputComponent* PlayerInputComponent)
+bool ATP_FirstPersonCharacter::GetHasRifle()
 {
-	if (FPlatformMisc::SupportsTouchInput() || GetDefault<UInputSettings>()->bUseMouseForTouch)
-	{
-		PlayerInputComponent->BindTouch(EInputEvent::IE_Pressed, this, &ATP_FirstPersonCharacter::BeginTouch);
-		PlayerInputComponent->BindTouch(EInputEvent::IE_Released, this, &ATP_FirstPersonCharacter::EndTouch);
-
-		return true;
-	}
-	
-	return false;
-}
+	return bHasRifle;
+}
\ No newline at end of file
diff --git a/Templates/TP_FirstPerson/Source/TP_FirstPerson/TP_FirstPersonCharacter.h b/Templates/TP_FirstPerson/Source/TP_FirstPerson/TP_FirstPersonCharacter.h
--- a/Templates/TP_FirstPerson/Source/TP_FirstPerson/TP_FirstPersonCharacter.h
+++ b/Templates/TP_FirstPerson/Source/TP_FirstPerson/TP_FirstPersonCharacter.h
@@ -4,6 +4,7 @@
 
 #include "CoreMinimal.h"
 #include "GameFramework/Character.h"
+#include "InputActionValue.h"
 #include "TP_FirstPersonCharacter.generated.h"
 
 class UInputComponent;
@@ -13,10 +14,6 @@ class UCameraComponent;
 class UAnimMontage;
 class USoundBase;
 
-// Declaration of the delegate that will be called when the Primary Action is triggered
-// It is declared as dynamic so it can be accessed also in Blueprints
-DECLARE_DYNAMIC_MULTICAST_DELEGATE(FOnUseItem);
-
 UCLASS(config=Game)
 class ATP_FirstPersonCharacter : public ACharacter
 {
@@ -30,6 +27,19 @@ class ATP_FirstPersonCharacter : public ACharacter
 	UPROPERTY(VisibleAnywhere, BlueprintReadOnly, Category = Camera, meta = (AllowPrivateAccess = "true"))
 	UCameraComponent* FirstPersonCameraComponent;
 
+	/** MappingContext */
+	UPROPERTY(EditAnywhere, BlueprintReadOnly, Category=Input, meta=(AllowPrivateAccess = "true"))
+	class UInputMappingContext* DefaultMappingContext;
+
+	/** Jump Input Action */
+	UPROPERTY(EditAnywhere, BlueprintReadOnly, Category=Input, meta=(AllowPrivateAccess = "true"))
+	class UInputAction* JumpAction;
+
+	/** Move Input Action */
+	UPROPERTY(EditAnywhere, BlueprintReadOnly, Category=Input, meta=(AllowPrivateAccess = "true"))
+	class UInputAction* MoveAction;
+
+	
 public:
 	ATP_FirstPersonCharacter();
 
@@ -37,67 +47,45 @@ protected:
 	virtual void BeginPlay();
 
 public:
-	/** Base turn rate, in deg/sec. Other scaling may affect final turn rate. */
-	UPROPERTY(VisibleAnywhere, BlueprintReadOnly, Category=Camera)
-	float TurnRateGamepad;
+	
+	/** Called for adding mapping context*/
+	virtual void PossessedBy(AController* NewController) override;
+	virtual void UnPossessed() override;
+
+	/** Look Input Action */
+	UPROPERTY(EditAnywhere, BlueprintReadOnly, Category = Input, meta = (AllowPrivateAccess = "true"))
+	class UInputAction* LookAction;
+
+	/** Bool for AnimBP to switch to another animation set */
+	UPROPERTY(VisibleAnywhere, BlueprintReadOnly, Category = Weapon)
+	bool bHasRifle;
+
+	/** Setter to set the bool */
+	UFUNCTION(BlueprintCallable, Category = Weapon)
+	void SetHasRifle(bool bNewHasRifle);
+
+	/** Getter for the bool */
+	UFUNCTION(BlueprintCallable, Category = Weapon)
+	bool GetHasRifle();
 
-	/** Delegate to whom anyone can subscribe to receive this event */
-	UPROPERTY(BlueprintAssignable, Category = "Interaction")
-	FOnUseItem OnUseItem;
 protected:
-	
-	/** Fires a projectile. */
-	void OnPrimaryAction();
-
-	/** Handles moving forward/backward */
-	void MoveForward(float Val);
-
-	/** Handles strafing movement, left and right */
-	void MoveRight(float Val);
-
-	/**
-	 * Called via input to turn at a given rate.
-	 * @param Rate	This is a normalized rate, i.e. 1.0 means 100% of desired turn rate
-	 */
-	void TurnAtRate(float Rate);
-
-	/**
-	 * Called via input to turn look up/down at a given rate.
-	 * @param Rate	This is a normalized rate, i.e. 1.0 means 100% of desired turn rate
-	 */
-	void LookUpAtRate(float Rate);
-
-	struct TouchData
-	{
-		TouchData() { bIsPressed = false;Location=FVector::ZeroVector;}
-		bool bIsPressed;
-		ETouchIndex::Type FingerIndex;
-		FVector Location;
-		bool bMoved;
-	};
-	void BeginTouch(const ETouchIndex::Type FingerIndex, const FVector Location);
-	void EndTouch(const ETouchIndex::Type FingerIndex, const FVector Location);
-	void TouchUpdate(const ETouchIndex::Type FingerIndex, const FVector Location);
-	TouchData	TouchItem;
-	
+	/** Called for movement input */
+	void Move(const FInputActionValue& Value);
+
+	/** Called for looking input */
+	void Look(const FInputActionValue& Value);
+
 protected:
 	// APawn interface
 	virtual void SetupPlayerInputComponent(UInputComponent* InputComponent) override;
 	// End of APawn interface
 
-	/* 
-	 * Configures input for touchscreen devices if there is a valid touch interface for doing so 
-	 *
-	 * @param	InputComponent	The input component pointer to bind controls to
-	 * @returns true if touch controls were enabled.
-	 */
-	bool EnableTouchscreenMovement(UInputComponent* InputComponent);
-
 public:
 	/** Returns Mesh1P subobject **/
 	USkeletalMeshComponent* GetMesh1P() const { return Mesh1P; }
 	/** Returns FirstPersonCameraComponent subobject **/
 	UCameraComponent* GetFirstPersonCameraComponent() const { return FirstPersonCameraComponent; }
 
+
 };
 
diff --git a/Templates/TP_FirstPerson/Source/TP_FirstPerson/TP_WeaponComponent.cpp b/Templates/TP_FirstPerson/Source/TP_FirstPerson/TP_WeaponComponent.cpp
--- a/Templates/TP_FirstPerson/Source/TP_FirstPerson/TP_WeaponComponent.cpp
+++ b/Templates/TP_FirstPerson/Source/TP_FirstPerson/TP_WeaponComponent.cpp
@@ -7,6 +7,8 @@
 #include "GameFramework/PlayerController.h"
 #include "Camera/PlayerCameraManager.h"
 #include "Kismet/GameplayStatics.h"
+#include "EnhancedInputComponent.h"
+#include "EnhancedInputSubsystems.h"
 
 // Sets default values for this component's properties
 UTP_WeaponComponent::UTP_WeaponComponent()
@@ -18,7 +20,7 @@ UTP_WeaponComponent::UTP_WeaponComponent()
 
 void UTP_WeaponComponent::Fire()
 {
-	if(Character == nullptr || Character->GetController() == nullptr)
+	if (Character == nullptr || Character->GetController() == nullptr)
 	{
 		return;
 	}
@@ -61,26 +63,50 @@ void UTP_WeaponComponent::Fire()
 	}
 }
 
-void UTP_WeaponComponent::EndPlay(const EEndPlayReason::Type EndPlayReason)
+void UTP_WeaponComponent::AttachWeapon(ATP_FirstPersonCharacter* TargetCharacter)
 {
-	if(Character != nullptr)
+	Character = TargetCharacter;
+	if (Character == nullptr)
 	{
-		// Unregister from the OnUseItem Event
-		Character->OnUseItem.RemoveDynamic(this, &UTP_WeaponComponent::Fire);
+		return;
 	}
-}
 
-void UTP_WeaponComponent::AttachWeapon(ATP_FirstPersonCharacter* TargetCharacter)
-{
-	Character = TargetCharacter;
-	if(Character != nullptr)
+	// Attach the weapon to the First Person Character
+	FAttachmentTransformRules AttachmentRules(EAttachmentRule::SnapToTarget, true);
+	AttachToComponent(Character->GetMesh1P(), AttachmentRules, FName(TEXT("GripPoint")));
+	
+	// switch bHasRifle so the animation blueprint can switch to another animation set
+	Character->SetHasRifle(true);
+
+	// Set up action bindings
+	if (APlayerController* PlayerController = Cast<APlayerController>(Character->GetController()))
 	{
-		// Attach the weapon to the First Person Character
-		FAttachmentTransformRules AttachmentRules(EAttachmentRule::SnapToTarget, true);
-		GetOwner()->AttachToComponent(Character->GetMesh1P(),AttachmentRules, FName(TEXT("GripPoint")));
+		if (UEnhancedInputLocalPlayerSubsystem* Subsystem = ULocalPlayer::GetSubsystem<UEnhancedInputLocalPlayerSubsystem>(PlayerController->GetLocalPlayer()))
+		{
+			// Set the priority of the mapping to 1, so that it overrides the Jump action with the Fire action when using touch input
+			Subsystem->AddMappingContext(FireMappingContext, 1);
+		}
 
-		// Register so that Fire is called every time the character tries to use the item being held
-		Character->OnUseItem.AddDynamic(this, &UTP_WeaponComponent::Fire);
+		if (UEnhancedInputComponent* EnhancedInputComponent = CastChecked<UEnhancedInputComponent>(PlayerController->InputComponent))
+		{
+			// Fire
+			EnhancedInputComponent->BindAction(FireAction, ETriggerEvent::Triggered, this, &UTP_WeaponComponent::Fire);
+		}
 	}
 }
 
+void UTP_WeaponComponent::EndPlay(const EEndPlayReason::Type EndPlayReason)
+{
+	if (Character == nullptr)
+	{
+		return;
+	}
+
+	if (APlayerController* PlayerController = Cast<APlayerController>(Character->GetController()))
+	{
+		if (UEnhancedInputLocalPlayerSubsystem* Subsystem = ULocalPlayer::GetSubsystem<UEnhancedInputLocalPlayerSubsystem>(PlayerController->GetLocalPlayer()))
+		{
+			Subsystem->RemoveMappingContext(FireMappingContext);
+		}
+	}
+}
\ No newline at end of file
diff --git a/Templates/TP_FirstPerson/Source/TP_FirstPerson/TP_WeaponComponent.h b/Templates/TP_FirstPerson/Source/TP_FirstPerson/TP_WeaponComponent.h
--- a/Templates/TP_FirstPerson/Source/TP_FirstPerson/TP_WeaponComponent.h
+++ b/Templates/TP_FirstPerson/Source/TP_FirstPerson/TP_WeaponComponent.h
@@ -3,13 +3,13 @@
 #pragma once
 
 #include "CoreMinimal.h"
-#include "Components/ActorComponent.h"
+#include "Components/SkeletalMeshComponent.h"
 #include "TP_WeaponComponent.generated.h"
 
 class ATP_FirstPersonCharacter;
 
 UCLASS(Blueprintable, BlueprintType, ClassGroup=(Custom), meta=(BlueprintSpawnableComponent) )
-class TP_FIRSTPERSON_API UTP_WeaponComponent : public UActorComponent
+class TP_FIRSTPERSON_API UTP_WeaponComponent : public USkeletalMeshComponent
 {
 	GENERATED_BODY()
 
@@ -30,6 +30,14 @@ public:
 	UPROPERTY(EditAnywhere, BlueprintReadWrite, Category=Gameplay)
 	FVector MuzzleOffset;
 
+	/** MappingContext */
+	UPROPERTY(EditAnywhere, BlueprintReadOnly, Category=Input, meta=(AllowPrivateAccess = "true"))
+	class UInputMappingContext* FireMappingContext;
+
+	/** Fire Input Action */
+	UPROPERTY(EditAnywhere, BlueprintReadOnly, Category=Input, meta=(AllowPrivateAccess = "true"))
+	class UInputAction* FireAction;
+
 	/** Sets default values for this component's properties */
 	UTP_WeaponComponent();
 
@@ -45,7 +53,6 @@ protected:
 	/** Ends gameplay for this component. */
 	UFUNCTION()
 	virtual void EndPlay(const EEndPlayReason::Type EndPlayReason) override;
-		
 
 private:
 	/** The Character holding this weapon*/
diff --git a/Templates/TP_FirstPersonBP/Config/DefaultEngine.ini b/Templates/TP_FirstPersonBP/Config/DefaultEngine.ini
--- a/Templates/TP_FirstPersonBP/Config/DefaultEngine.ini
+++ b/Templates/TP_FirstPersonBP/Config/DefaultEngine.ini
@@ -21,4 +21,19 @@ r.ReflectionMethod=1
 r.GenerateMeshDistanceFields=True
 r.DynamicGlobalIlluminationMethod=1
 r.Mobile.EnableNoPrecomputedLightingCSMShader=1
+r.DefaultFeature.AutoExposure.ExtendDefaultLuminanceRange=true
+
+[/Script/AndroidFileServerEditor.AndroidFileServerRuntimeSettings]
+bEnablePlugin=True
+bAllowNetworkConnection=True
+SecurityToken=5011EBFD4DF26CEF23A820976F45E941
+bIncludeInShipping=False
+bAllowExternalStartInShipping=False
+bCompileAFSProject=False
+bUseCompression=False
+bLogFiles=False
+bReportStats=False
+ConnectionType=USBOnly
+bUseManualIPAddress=False
+ManualIPAddress=
 
diff --git a/Templates/TP_ThirdPerson/Source/TP_ThirdPerson/TP_ThirdPersonCharacter.cpp b/Templates/TP_ThirdPerson/Source/TP_ThirdPerson/TP_ThirdPersonCharacter.cpp
--- a/Templates/TP_ThirdPerson/Source/TP_ThirdPerson/TP_ThirdPersonCharacter.cpp
+++ b/Templates/TP_ThirdPerson/Source/TP_ThirdPerson/TP_ThirdPersonCharacter.cpp
@@ -53,6 +53,8 @@ ATP_ThirdPersonCharacter::ATP_ThirdPersonCharacter()
 
 void ATP_ThirdPersonCharacter::PossessedBy(AController* NewController)
 {
+	Super::PossessedBy(NewController);
+
 	//Add a mapping context on Possessed
 	if (APlayerController* PlayerController = Cast<APlayerController>(NewController))
 	{
@@ -63,14 +65,24 @@ void ATP_ThirdPersonCharacter::PossessedBy(AController* NewController)
 	}
 }
 
+void ATP_ThirdPersonCharacter::UnPossessed()
+{
+	Super::UnPossessed();
+
+	if (APlayerController* PlayerController = Cast<APlayerController>(GetController()))
+	{
+		if (UEnhancedInputLocalPlayerSubsystem* Subsystem = ULocalPlayer::GetSubsystem<UEnhancedInputLocalPlayerSubsystem>(PlayerController->GetLocalPlayer()))
+		{
+			Subsystem->RemoveMappingContext(DefaultMappingContext);
+		}
+	}
+}
+
 //////////////////////////////////////////////////////////////////////////
 // Input
 
 void ATP_ThirdPersonCharacter::SetupPlayerInputComponent(class UInputComponent* PlayerInputComponent)
 {
-	// Set up gameplay key bindings
-	check(PlayerInputComponent);
-	
 	// Set up action bindings
 	if (UEnhancedInputComponent* EnhancedInputComponent = CastChecked<UEnhancedInputComponent>(PlayerInputComponent)) {
 		
diff --git a/Templates/TP_ThirdPerson/Source/TP_ThirdPerson/TP_ThirdPersonCharacter.h b/Templates/TP_ThirdPerson/Source/TP_ThirdPerson/TP_ThirdPersonCharacter.h
--- a/Templates/TP_ThirdPerson/Source/TP_ThirdPerson/TP_ThirdPersonCharacter.h
+++ b/Templates/TP_ThirdPerson/Source/TP_ThirdPerson/TP_ThirdPersonCharacter.h
@@ -41,8 +41,8 @@ public:
 	ATP_ThirdPersonCharacter();
 
 	/** Called for adding mapping context*/
-	void PossessedBy(AController* NewController) override;
-		
+	virtual void PossessedBy(AController* NewController) override;
+	virtual void UnPossessed() override;
 
 protected:
 
diff --git a/Templates/TP_ThirdPersonBP/Config/DefaultEngine.ini b/Templates/TP_ThirdPersonBP/Config/DefaultEngine.ini
--- a/Templates/TP_ThirdPersonBP/Config/DefaultEngine.ini
+++ b/Templates/TP_ThirdPersonBP/Config/DefaultEngine.ini
@@ -18,6 +18,7 @@ r.DynamicGlobalIlluminationMethod=1
 r.Lumen.TraceMeshSDFs=0
 r.Shadow.Virtual.Enable=1
 r.Mobile.EnableNoPrecomputedLightingCSMShader=1
+r.DefaultFeature.AutoExposure.ExtendDefaultLuminanceRange=true
 
 [/Script/WindowsTargetPlatform.WindowsTargetSettings]
 DefaultGraphicsRHI=DefaultGraphicsRHI_Default
diff --git a/Templates/TP_TopDown/Config/DefaultEngine.ini b/Templates/TP_TopDown/Config/DefaultEngine.ini
--- a/Templates/TP_TopDown/Config/DefaultEngine.ini
+++ b/Templates/TP_TopDown/Config/DefaultEngine.ini
@@ -4,4 +4,49 @@ EditorStartupMap=/Game/TopDown/Maps/TopDownMap.TopDownMap
 GlobalDefaultGameMode="/Script/TP_TopDown.TP_TopDownGameMode"
 
 [/Script/Engine.RendererSettings]
-r.Mobile.EnableNoPrecomputedLightingCSMShader=1
\ No newline at end of file
+r.Mobile.EnableNoPrecomputedLightingCSMShader=1
+ r.DefaultFeature.AutoExposure.ExtendDefaultLuminanceRange=true
+
+[/Script/AndroidFileServerEditor.AndroidFileServerRuntimeSettings]
+bEnablePlugin=True
+bAllowNetworkConnection=True
+SecurityToken=60F4C6BA45270F82910E719121C45142
+bIncludeInShipping=False
+bAllowExternalStartInShipping=False
+bCompileAFSProject=False
+bUseCompression=False
+bLogFiles=False
+bReportStats=False
+ConnectionType=USBOnly
+bUseManualIPAddress=False
+ManualIPAddress=
+
+[/Script/WindowsTargetPlatform.WindowsTargetSettings]
+DefaultGraphicsRHI=DefaultGraphicsRHI_Default
+-D3D12TargetedShaderFormats=PCD3D_SM5
++D3D12TargetedShaderFormats=PCD3D_SM6
+-D3D11TargetedShaderFormats=PCD3D_SM5
++D3D11TargetedShaderFormats=PCD3D_SM5
+Compiler=Default
+AudioSampleRate=48000
+AudioCallbackBufferFrameSize=1024
+AudioNumBuffersToEnqueue=1
+AudioMaxChannels=0
+AudioNumSourceWorkers=4
+SpatializationPlugin=
+SourceDataOverridePlugin=
+ReverbPlugin=
+OcclusionPlugin=
+CompressionOverrides=(bOverrideCompressionTimes=False,DurationThreshold=5.000000,MaxNumRandomBranches=0,SoundCueQualityIndex=0)
+CacheSizeKB=65536
+MaxChunkSizeOverrideKB=0
+bResampleForDevice=False
+MaxSampleRate=48000.000000
+HighSampleRate=32000.000000
+MedSampleRate=24000.000000
+LowSampleRate=12000.000000
+MinSampleRate=8000.000000
+CompressionQualityModifier=1.000000
+AutoStreamingThreshold=0.000000
+SoundCueCookQualityIndex=-1
+
diff --git a/Templates/TP_TopDown/Config/DefaultInput.ini b/Templates/TP_TopDown/Config/DefaultInput.ini
--- a/Templates/TP_TopDown/Config/DefaultInput.ini
+++ b/Templates/TP_TopDown/Config/DefaultInput.ini
@@ -18,6 +18,46 @@
 +AxisConfig=(AxisKeyName="Gamepad_RightTriggerAxis",AxisProperties=(DeadZone=0.000000,Sensitivity=1.000000,Exponent=1.000000,bInvert=False))
 +AxisConfig=(AxisKeyName="Gamepad_Special_Left_X",AxisProperties=(DeadZone=0.000000,Sensitivity=1.000000,Exponent=1.000000,bInvert=False))
 +AxisConfig=(AxisKeyName="Gamepad_Special_Left_Y",AxisProperties=(DeadZone=0.000000,Sensitivity=1.000000,Exponent=1.000000,bInvert=False))
++AxisConfig=(AxisKeyName="Vive_Left_Trigger_Axis",AxisProperties=(DeadZone=0.000000,Sensitivity=1.000000,Exponent=1.000000,bInvert=False))
++AxisConfig=(AxisKeyName="Vive_Left_Trackpad_X",AxisProperties=(DeadZone=0.000000,Sensitivity=1.000000,Exponent=1.000000,bInvert=False))
++AxisConfig=(AxisKeyName="Vive_Left_Trackpad_Y",AxisProperties=(DeadZone=0.000000,Sensitivity=1.000000,Exponent=1.000000,bInvert=False))
++AxisConfig=(AxisKeyName="Vive_Right_Trigger_Axis",AxisProperties=(DeadZone=0.000000,Sensitivity=1.000000,Exponent=1.000000,bInvert=False))
++AxisConfig=(AxisKeyName="Vive_Right_Trackpad_X",AxisProperties=(DeadZone=0.000000,Sensitivity=1.000000,Exponent=1.000000,bInvert=False))
++AxisConfig=(AxisKeyName="Vive_Right_Trackpad_Y",AxisProperties=(DeadZone=0.000000,Sensitivity=1.000000,Exponent=1.000000,bInvert=False))
++AxisConfig=(AxisKeyName="MixedReality_Left_Trigger_Axis",AxisProperties=(DeadZone=0.000000,Sensitivity=1.000000,Exponent=1.000000,bInvert=False))
++AxisConfig=(AxisKeyName="MixedReality_Left_Thumbstick_X",AxisProperties=(DeadZone=0.000000,Sensitivity=1.000000,Exponent=1.000000,bInvert=False))
++AxisConfig=(AxisKeyName="MixedReality_Left_Thumbstick_Y",AxisProperties=(DeadZone=0.000000,Sensitivity=1.000000,Exponent=1.000000,bInvert=False))
++AxisConfig=(AxisKeyName="MixedReality_Left_Trackpad_X",AxisProperties=(DeadZone=0.000000,Sensitivity=1.000000,Exponent=1.000000,bInvert=False))
++AxisConfig=(AxisKeyName="MixedReality_Left_Trackpad_Y",AxisProperties=(DeadZone=0.000000,Sensitivity=1.000000,Exponent=1.000000,bInvert=False))
++AxisConfig=(AxisKeyName="MixedReality_Right_Trigger_Axis",AxisProperties=(DeadZone=0.000000,Sensitivity=1.000000,Exponent=1.000000,bInvert=False))
++AxisConfig=(AxisKeyName="MixedReality_Right_Thumbstick_X",AxisProperties=(DeadZone=0.000000,Sensitivity=1.000000,Exponent=1.000000,bInvert=False))
++AxisConfig=(AxisKeyName="MixedReality_Right_Thumbstick_Y",AxisProperties=(DeadZone=0.000000,Sensitivity=1.000000,Exponent=1.000000,bInvert=False))
++AxisConfig=(AxisKeyName="MixedReality_Right_Trackpad_X",AxisProperties=(DeadZone=0.000000,Sensitivity=1.000000,Exponent=1.000000,bInvert=False))
++AxisConfig=(AxisKeyName="MixedReality_Right_Trackpad_Y",AxisProperties=(DeadZone=0.000000,Sensitivity=1.000000,Exponent=1.000000,bInvert=False))
++AxisConfig=(AxisKeyName="OculusTouch_Left_Grip_Axis",AxisProperties=(DeadZone=0.000000,Sensitivity=1.000000,Exponent=1.000000,bInvert=False))
++AxisConfig=(AxisKeyName="OculusTouch_Left_Trigger_Axis",AxisProperties=(DeadZone=0.000000,Sensitivity=1.000000,Exponent=1.000000,bInvert=False))
++AxisConfig=(AxisKeyName="OculusTouch_Left_Thumbstick_X",AxisProperties=(DeadZone=0.000000,Sensitivity=1.000000,Exponent=1.000000,bInvert=False))
++AxisConfig=(AxisKeyName="OculusTouch_Left_Thumbstick_Y",AxisProperties=(DeadZone=0.000000,Sensitivity=1.000000,Exponent=1.000000,bInvert=False))
++AxisConfig=(AxisKeyName="OculusTouch_Right_Grip_Axis",AxisProperties=(DeadZone=0.000000,Sensitivity=1.000000,Exponent=1.000000,bInvert=False))
++AxisConfig=(AxisKeyName="OculusTouch_Right_Trigger_Axis",AxisProperties=(DeadZone=0.000000,Sensitivity=1.000000,Exponent=1.000000,bInvert=False))
++AxisConfig=(AxisKeyName="OculusTouch_Right_Thumbstick_X",AxisProperties=(DeadZone=0.000000,Sensitivity=1.000000,Exponent=1.000000,bInvert=False))
++AxisConfig=(AxisKeyName="OculusTouch_Right_Thumbstick_Y",AxisProperties=(DeadZone=0.000000,Sensitivity=1.000000,Exponent=1.000000,bInvert=False))
++AxisConfig=(AxisKeyName="ValveIndex_Left_Grip_Axis",AxisProperties=(DeadZone=0.000000,Sensitivity=1.000000,Exponent=1.000000,bInvert=False))
++AxisConfig=(AxisKeyName="ValveIndex_Left_Grip_Force",AxisProperties=(DeadZone=0.000000,Sensitivity=1.000000,Exponent=1.000000,bInvert=False))
++AxisConfig=(AxisKeyName="ValveIndex_Left_Trigger_Axis",AxisProperties=(DeadZone=0.000000,Sensitivity=1.000000,Exponent=1.000000,bInvert=False))
++AxisConfig=(AxisKeyName="ValveIndex_Left_Thumbstick_X",AxisProperties=(DeadZone=0.000000,Sensitivity=1.000000,Exponent=1.000000,bInvert=False))
++AxisConfig=(AxisKeyName="ValveIndex_Left_Thumbstick_Y",AxisProperties=(DeadZone=0.000000,Sensitivity=1.000000,Exponent=1.000000,bInvert=False))
++AxisConfig=(AxisKeyName="ValveIndex_Left_Trackpad_X",AxisProperties=(DeadZone=0.000000,Sensitivity=1.000000,Exponent=1.000000,bInvert=False))
++AxisConfig=(AxisKeyName="ValveIndex_Left_Trackpad_Y",AxisProperties=(DeadZone=0.000000,Sensitivity=1.000000,Exponent=1.000000,bInvert=False))
++AxisConfig=(AxisKeyName="ValveIndex_Left_Trackpad_Force",AxisProperties=(DeadZone=0.000000,Sensitivity=1.000000,Exponent=1.000000,bInvert=False))
++AxisConfig=(AxisKeyName="ValveIndex_Right_Grip_Axis",AxisProperties=(DeadZone=0.000000,Sensitivity=1.000000,Exponent=1.000000,bInvert=False))
++AxisConfig=(AxisKeyName="ValveIndex_Right_Grip_Force",AxisProperties=(DeadZone=0.000000,Sensitivity=1.000000,Exponent=1.000000,bInvert=False))
++AxisConfig=(AxisKeyName="ValveIndex_Right_Trigger_Axis",AxisProperties=(DeadZone=0.000000,Sensitivity=1.000000,Exponent=1.000000,bInvert=False))
++AxisConfig=(AxisKeyName="ValveIndex_Right_Thumbstick_X",AxisProperties=(DeadZone=0.000000,Sensitivity=1.000000,Exponent=1.000000,bInvert=False))
++AxisConfig=(AxisKeyName="ValveIndex_Right_Thumbstick_Y",AxisProperties=(DeadZone=0.000000,Sensitivity=1.000000,Exponent=1.000000,bInvert=False))
++AxisConfig=(AxisKeyName="ValveIndex_Right_Trackpad_X",AxisProperties=(DeadZone=0.000000,Sensitivity=1.000000,Exponent=1.000000,bInvert=False))
++AxisConfig=(AxisKeyName="ValveIndex_Right_Trackpad_Y",AxisProperties=(DeadZone=0.000000,Sensitivity=1.000000,Exponent=1.000000,bInvert=False))
++AxisConfig=(AxisKeyName="ValveIndex_Right_Trackpad_Force",AxisProperties=(DeadZone=0.000000,Sensitivity=1.000000,Exponent=1.000000,bInvert=False))
 bAltEnterTogglesFullscreen=True
 bF11TogglesFullscreen=True
 bUseMouseForTouch=False
@@ -25,6 +65,9 @@ bEnableMouseSmoothing=True
 bEnableFOVScaling=True
 bCaptureMouseOnLaunch=True
 bEnableLegacyInputScales=True
+bEnableMotionControls=True
+bFilterInputByPlatformUser=False
+bShouldFlushPressedKeysOnViewportFocusLost=True
 bAlwaysShowTouchInterface=False
 bShowConsoleOnFourFingerTap=True
 bEnableGestureRecognizer=False
@@ -33,7 +76,6 @@ DefaultViewportMouseCaptureMode=CapturePermanently_IncludingInitialMouseDown
 DefaultViewportMouseLockMode=LockOnCapture
 FOVScale=0.011110
 DoubleClickTime=0.200000
-+ActionMappings=(ActionName="SetDestination",bShift=False,bCtrl=False,bAlt=False,bCmd=False,Key=LeftMouseButton)
 DefaultPlayerInputClass=/Script/EnhancedInput.EnhancedPlayerInput
 DefaultInputComponentClass=/Script/EnhancedInput.EnhancedInputComponent
 DefaultTouchInterface=None
diff --git a/Templates/TP_TopDown/Source/TP_TopDown/TP_TopDown.Build.cs b/Templates/TP_TopDown/Source/TP_TopDown/TP_TopDown.Build.cs
--- a/Templates/TP_TopDown/Source/TP_TopDown/TP_TopDown.Build.cs
+++ b/Templates/TP_TopDown/Source/TP_TopDown/TP_TopDown.Build.cs
@@ -8,6 +8,6 @@ public class TP_TopDown : ModuleRules
 	{
 		PCHUsage = PCHUsageMode.UseExplicitOrSharedPCHs;
 
-        PublicDependencyModuleNames.AddRange(new string[] { "Core", "CoreUObject", "Engine", "InputCore", "HeadMountedDisplay", "NavigationSystem", "AIModule", "Niagara" });
+        PublicDependencyModuleNames.AddRange(new string[] { "Core", "CoreUObject", "Engine", "InputCore", "HeadMountedDisplay", "NavigationSystem", "AIModule", "Niagara", "EnhancedInput" });
     }
 }
diff --git a/Templates/TP_TopDown/Source/TP_TopDown/TP_TopDownPlayerController.cpp b/Templates/TP_TopDown/Source/TP_TopDown/TP_TopDownPlayerController.cpp
--- a/Templates/TP_TopDown/Source/TP_TopDown/TP_TopDownPlayerController.cpp
+++ b/Templates/TP_TopDown/Source/TP_TopDown/TP_TopDownPlayerController.cpp
@@ -7,45 +7,35 @@
 #include "NiagaraFunctionLibrary.h"
 #include "TP_TopDownCharacter.h"
 #include "Engine/World.h"
+#include "EnhancedInputComponent.h"
+#include "EnhancedInputSubsystems.h"
 
 ATP_TopDownPlayerController::ATP_TopDownPlayerController()
 {
 	bShowMouseCursor = true;
 	DefaultMouseCursor = EMouseCursor::Default;
+	CachedDestination = FVector::ZeroVector;
+	FollowTime = 0.f;
 }
 
-void ATP_TopDownPlayerController::PlayerTick(float DeltaTime)
+void ATP_TopDownPlayerController::OnPossess(APawn* PossessedPawn)
 {
-	Super::PlayerTick(DeltaTime);
+	Super::OnPossess(PossessedPawn);
 
-	if(bInputPressed)
+	if (UEnhancedInputLocalPlayerSubsystem* Subsystem = ULocalPlayer::GetSubsystem<UEnhancedInputLocalPlayerSubsystem>(GetLocalPlayer()))
 	{
-		FollowTime += DeltaTime;
-
-		// Look for the touch location
-		FVector HitLocation = FVector::ZeroVector;
-		FHitResult Hit;
-		if(bIsTouch)
-		{
-			GetHitResultUnderFinger(ETouchIndex::Touch1, ECC_Visibility, true, Hit);
-		}
-		else
-		{
-			GetHitResultUnderCursor(ECC_Visibility, true, Hit);
-		}
-		HitLocation = Hit.Location;
-
-		// Direct the Pawn towards that location
-		APawn* const MyPawn = GetPawn();
-		if(MyPawn)
-		{
-			FVector WorldDirection = (HitLocation - MyPawn->GetActorLocation()).GetSafeNormal();
-			MyPawn->AddMovementInput(WorldDirection, 1.f, false);
-		}
+		Subsystem->AddMappingContext(DefaultMappingContext, 0);
 	}
-	else
+}
+
+void ATP_TopDownPlayerController::OnUnPossess()
+{
+	Super::OnUnPossess();
+
+	if (UEnhancedInputLocalPlayerSubsystem* Subsystem = ULocalPlayer::GetSubsystem<UEnhancedInputLocalPlayerSubsystem>(GetLocalPlayer()))
 	{
-		FollowTime = 0.f;
+		check(DefaultMappingContext);
+		Subsystem->RemoveMappingContext(DefaultMappingContext);
 	}
 }
 
@@ -54,50 +44,82 @@ void ATP_TopDownPlayerController::SetupInputComponent()
 	// set up gameplay key bindings
 	Super::SetupInputComponent();
 
-	InputComponent->BindAction("SetDestination", IE_Pressed, this, &ATP_TopDownPlayerController::OnSetDestinationPressed);
-	InputComponent->BindAction("SetDestination", IE_Released, this, &ATP_TopDownPlayerController::OnSetDestinationReleased);
-
-	// support touch devices 
-	InputComponent->BindTouch(EInputEvent::IE_Pressed, this, &ATP_TopDownPlayerController::OnTouchPressed);
-	InputComponent->BindTouch(EInputEvent::IE_Released, this, &ATP_TopDownPlayerController::OnTouchReleased);
+	// Set up action bindings
+	if (UEnhancedInputComponent* EnhancedInputComponent = CastChecked<UEnhancedInputComponent>(InputComponent))
+	{
+		// Setup mouse input events
+		EnhancedInputComponent->BindAction(SetDestinationClickAction, ETriggerEvent::Started, this, &ATP_TopDownPlayerController::OnInputStarted);
+		EnhancedInputComponent->BindAction(SetDestinationClickAction, ETriggerEvent::Triggered, this, &ATP_TopDownPlayerController::OnSetDestinationTriggered);
+		EnhancedInputComponent->BindAction(SetDestinationClickAction, ETriggerEvent::Completed, this, &ATP_TopDownPlayerController::OnSetDestinationReleased);
+		EnhancedInputComponent->BindAction(SetDestinationClickAction, ETriggerEvent::Canceled, this, &ATP_TopDownPlayerController::OnSetDestinationReleased);
+
+		// Setup touch input events
+		EnhancedInputComponent->BindAction(SetDestinationTouchAction, ETriggerEvent::Started, this, &ATP_TopDownPlayerController::OnInputStarted);
+		EnhancedInputComponent->BindAction(SetDestinationTouchAction, ETriggerEvent::Triggered, this, &ATP_TopDownPlayerController::OnTouchTriggered);
+		EnhancedInputComponent->BindAction(SetDestinationTouchAction, ETriggerEvent::Completed, this, &ATP_TopDownPlayerController::OnTouchReleased);
+		EnhancedInputComponent->BindAction(SetDestinationTouchAction, ETriggerEvent::Canceled, this, &ATP_TopDownPlayerController::OnTouchReleased);
+	}
+}
 
+void ATP_TopDownPlayerController::OnInputStarted()
+{
+	StopMovement();
 }
 
-void ATP_TopDownPlayerController::OnSetDestinationPressed()
+// Triggered every frame when the input is held down
+void ATP_TopDownPlayerController::OnSetDestinationTriggered()
 {
 	// We flag that the input is being pressed
-	bInputPressed = true;
-	// Just in case the character was moving because of a previous short press we stop it
-	StopMovement();
+	FollowTime += GetWorld()->GetDeltaSeconds();
+	
+	// We look for the location in the world where the player has pressed the input
+	FHitResult Hit;
+	bool bHitSuccessful = false;
+	if (bIsTouch)
+	{
+		bHitSuccessful = GetHitResultUnderFinger(ETouchIndex::Touch1, ECollisionChannel::ECC_Visibility, true, Hit);
+	}
+	else
+	{
+		bHitSuccessful = GetHitResultUnderCursor(ECollisionChannel::ECC_Visibility, true, Hit);
+	}
+
+	// If we hit a surface, cache the location
+	if (bHitSuccessful)
+	{
+		CachedDestination = Hit.Location;
+	}
+	
+	// Move towards mouse pointer or touch
+	APawn* ControlledPawn = GetPawn();
+	if (ControlledPawn != nullptr)
+	{
+		FVector WorldDirection = (CachedDestination - ControlledPawn->GetActorLocation()).GetSafeNormal();
+		ControlledPawn->AddMovementInput(WorldDirection, 1.0, false);
+	}
 }
 
 void ATP_TopDownPlayerController::OnSetDestinationReleased()
 {
-	// Player is no longer pressing the input
-	bInputPressed = false;
-
 	// If it was a short press
-	if(FollowTime <= ShortPressThreshold)
+	if (FollowTime <= ShortPressThreshold)
 	{
-		// We look for the location in the world where the player has pressed the input
-		FVector HitLocation = FVector::ZeroVector;
-		FHitResult Hit;
-		GetHitResultUnderCursor(ECC_Visibility, true, Hit);
-		HitLocation = Hit.Location;
-
 		// We move there and spawn some particles
-		UAIBlueprintHelperLibrary::SimpleMoveToLocation(this, HitLocation);
-		UNiagaraFunctionLibrary::SpawnSystemAtLocation(this, FXCursor, HitLocation, FRotator::ZeroRotator, FVector(1.f, 1.f, 1.f), true, true, ENCPoolMethod::None, true);
+		UAIBlueprintHelperLibrary::SimpleMoveToLocation(this, CachedDestination);
+		UNiagaraFunctionLibrary::SpawnSystemAtLocation(this, FXCursor, CachedDestination, FRotator::ZeroRotator, FVector(1.f, 1.f, 1.f), true, true, ENCPoolMethod::None, true);
 	}
+
+	FollowTime = 0.f;
 }
 
-void ATP_TopDownPlayerController::OnTouchPressed(const ETouchIndex::Type FingerIndex, const FVector Location)
+// Triggered every frame when the input is held down
+void ATP_TopDownPlayerController::OnTouchTriggered()
 {
 	bIsTouch = true;
-	OnSetDestinationPressed();
+	OnSetDestinationTriggered();
 }
 
-void ATP_TopDownPlayerController::OnTouchReleased(const ETouchIndex::Type FingerIndex, const FVector Location)
+void ATP_TopDownPlayerController::OnTouchReleased()
 {
 	bIsTouch = false;
 	OnSetDestinationReleased();
diff --git a/Templates/TP_TopDown/Source/TP_TopDown/TP_TopDownPlayerController.h b/Templates/TP_TopDown/Source/TP_TopDown/TP_TopDownPlayerController.h
--- a/Templates/TP_TopDown/Source/TP_TopDown/TP_TopDownPlayerController.h
+++ b/Templates/TP_TopDown/Source/TP_TopDown/TP_TopDownPlayerController.h
@@ -5,6 +5,7 @@
 #include "CoreMinimal.h"
 #include "Templates/SubclassOf.h"
 #include "GameFramework/PlayerController.h"
+#include "InputActionValue.h"
 #include "TP_TopDownPlayerController.generated.h"
 
 /** Forward declaration to improve compiling times */
@@ -26,23 +27,38 @@ public:
 	UPROPERTY(EditAnywhere, BlueprintReadOnly, Category = Input)
 	UNiagaraSystem* FXCursor;
 
+	/** MappingContext */
+	UPROPERTY(EditAnywhere, BlueprintReadOnly, Category=Input, meta=(AllowPrivateAccess = "true"))
+	class UInputMappingContext* DefaultMappingContext;
+	
+	/** Jump Input Action */
+	UPROPERTY(EditAnywhere, BlueprintReadOnly, Category=Input, meta=(AllowPrivateAccess = "true"))
+	class UInputAction* SetDestinationClickAction;
+
+	/** Jump Input Action */
+	UPROPERTY(EditAnywhere, BlueprintReadOnly, Category=Input, meta=(AllowPrivateAccess = "true"))
+	class UInputAction* SetDestinationTouchAction;
+
 protected:
 	/** True if the controlled character should navigate to the mouse cursor. */
 	uint32 bMoveToMouseCursor : 1;
 
 	// Begin PlayerController interface
-	virtual void PlayerTick(float DeltaTime) override;
+	virtual void OnPossess(APawn* PossessedPawn) override;
+	virtual void OnUnPossess() override;
 	virtual void SetupInputComponent() override;
 	// End PlayerController interface
 
 	/** Input handlers for SetDestination action. */
-	void OnSetDestinationPressed();
+	void OnInputStarted();
+	void OnSetDestinationTriggered();
 	void OnSetDestinationReleased();
-	void OnTouchPressed(const ETouchIndex::Type FingerIndex, const FVector Location);
-	void OnTouchReleased(const ETouchIndex::Type FingerIndex, const FVector Location);
+	void OnTouchTriggered();
+	void OnTouchReleased();
 
 private:
-	bool bInputPressed; // Input is bring pressed
+	FVector CachedDestination;
+
 	bool bIsTouch; // Is it a touch device
 	float FollowTime; // For how long it has been pressed
 };
diff --git a/Templates/TP_TopDownBP/Config/DefaultEngine.ini b/Templates/TP_TopDownBP/Config/DefaultEngine.ini
--- a/Templates/TP_TopDownBP/Config/DefaultEngine.ini
+++ b/Templates/TP_TopDownBP/Config/DefaultEngine.ini
@@ -55,4 +55,48 @@ r.GenerateMeshDistanceFields=True
 r.DynamicGlobalIlluminationMethod=1
 r.Shadow.Virtual.Enable=1
 r.Mobile.EnableNoPrecomputedLightingCSMShader=1
+r.DefaultFeature.AutoExposure.ExtendDefaultLuminanceRange=true
+
+[/Script/AndroidFileServerEditor.AndroidFileServerRuntimeSettings]
+bEnablePlugin=True
+bAllowNetworkConnection=True
+SecurityToken=CAF39D9941794497898047B0405D7FA8
+bIncludeInShipping=False
+bAllowExternalStartInShipping=False
+bCompileAFSProject=False
+bUseCompression=False
+bLogFiles=False
+bReportStats=False
+ConnectionType=USBOnly
+bUseManualIPAddress=False
+ManualIPAddress=
+
+[/Script/WindowsTargetPlatform.WindowsTargetSettings]
+DefaultGraphicsRHI=DefaultGraphicsRHI_Default
+-D3D12TargetedShaderFormats=PCD3D_SM5
++D3D12TargetedShaderFormats=PCD3D_SM6
+-D3D11TargetedShaderFormats=PCD3D_SM5
++D3D11TargetedShaderFormats=PCD3D_SM5
+Compiler=Default
+AudioSampleRate=48000
+AudioCallbackBufferFrameSize=1024
+AudioNumBuffersToEnqueue=1
+AudioMaxChannels=0
+AudioNumSourceWorkers=4
+SpatializationPlugin=
+SourceDataOverridePlugin=
+ReverbPlugin=
+OcclusionPlugin=
+CompressionOverrides=(bOverrideCompressionTimes=False,DurationThreshold=5.000000,MaxNumRandomBranches=0,SoundCueQualityIndex=0)
+CacheSizeKB=65536
+MaxChunkSizeOverrideKB=0
+bResampleForDevice=False
+MaxSampleRate=48000.000000
+HighSampleRate=32000.000000
+MedSampleRate=24000.000000
+LowSampleRate=12000.000000
+MinSampleRate=8000.000000
+CompressionQualityModifier=1.000000
+AutoStreamingThreshold=0.000000
+SoundCueCookQualityIndex=-1
 
diff --git a/Templates/TP_VehicleAdvBP/Config/DefaultEngine.ini b/Templates/TP_VehicleAdvBP/Config/DefaultEngine.ini
--- a/Templates/TP_VehicleAdvBP/Config/DefaultEngine.ini
+++ b/Templates/TP_VehicleAdvBP/Config/DefaultEngine.ini
@@ -38,4 +38,34 @@ r.GenerateMeshDistanceFields=True
 r.DynamicGlobalIlluminationMethod=1
 r.Mobile.EnableNoPrecomputedLightingCSMShader=1
 r.ClearCoatNormal=True
+r.DefaultFeature.AutoExposure.ExtendDefaultLuminanceRange=true
+
+[/Script/WindowsTargetPlatform.WindowsTargetSettings]
+DefaultGraphicsRHI=DefaultGraphicsRHI_Default
+-D3D12TargetedShaderFormats=PCD3D_SM5
++D3D12TargetedShaderFormats=PCD3D_SM6
+-D3D11TargetedShaderFormats=PCD3D_SM5
++D3D11TargetedShaderFormats=PCD3D_SM5
+Compiler=Default
+AudioSampleRate=48000
+AudioCallbackBufferFrameSize=1024
+AudioNumBuffersToEnqueue=1
+AudioMaxChannels=0
+AudioNumSourceWorkers=4
+SpatializationPlugin=
+SourceDataOverridePlugin=
+ReverbPlugin=
+OcclusionPlugin=
+CompressionOverrides=(bOverrideCompressionTimes=False,DurationThreshold=5.000000,MaxNumRandomBranches=0,SoundCueQualityIndex=0)
+CacheSizeKB=65536
+MaxChunkSizeOverrideKB=0
+bResampleForDevice=False
+MaxSampleRate=48000.000000
+HighSampleRate=32000.000000
+MedSampleRate=24000.000000
+LowSampleRate=12000.000000
+MinSampleRate=8000.000000
+CompressionQualityModifier=1.000000
+AutoStreamingThreshold=0.000000
+SoundCueCookQualityIndex=-1
 
