[Templates] Add mapping input context moved to BeginPlay event so it works also for multiplayer. (Note: TopDown has SimpleMoveTo which would need a a more complex setup for multiplayer telling the server to move than we want for a template; holding the mouse with movemet component works tho) Updated the feature packs.

#content only

diff --git a/Templates/TP_FirstPerson/Source/TP_FirstPerson/TP_FirstPersonCharacter.cpp b/Templates/TP_FirstPerson/Source/TP_FirstPerson/TP_FirstPersonCharacter.cpp
--- a/Templates/TP_FirstPerson/Source/TP_FirstPerson/TP_FirstPersonCharacter.cpp
+++ b/Templates/TP_FirstPerson/Source/TP_FirstPerson/TP_FirstPersonCharacter.cpp
@@ -42,6 +42,15 @@ void ATP_FirstPersonCharacter::BeginPlay()
 	// Call the base class  
 	Super::BeginPlay();
 
+	//Add Input Mapping Context
+	if (APlayerController* PlayerController = Cast<APlayerController>(Controller))
+	{
+		if (UEnhancedInputLocalPlayerSubsystem* Subsystem = ULocalPlayer::GetSubsystem<UEnhancedInputLocalPlayerSubsystem>(PlayerController->GetLocalPlayer()))
+		{
+			Subsystem->AddMappingContext(DefaultMappingContext, 0);
+		}
+	}
+
 }
 
 //////////////////////////////////////////////////////////////////////////// Input
@@ -63,32 +72,6 @@ void ATP_FirstPersonCharacter::SetupPlayerInputComponent(class UInputComponent*
 	}
 }
 
-void ATP_FirstPersonCharacter::PossessedBy(AController* NewController)
-{
-	Super::PossessedBy(NewController);
-
-	//Add a mapping context on Possessed
-	if (APlayerController* PlayerController = Cast<APlayerController>(NewController))
-	{
-		if (UEnhancedInputLocalPlayerSubsystem* Subsystem = ULocalPlayer::GetSubsystem<UEnhancedInputLocalPlayerSubsystem>(PlayerController->GetLocalPlayer()))
-		{
-			Subsystem->AddMappingContext(DefaultMappingContext, 0);
-		}
-	}
-}
-
-void ATP_FirstPersonCharacter::UnPossessed()
-{
-	Super::UnPossessed();
-
-	if (APlayerController* PlayerController = Cast<APlayerController>(GetController()))
-	{
-		if (UEnhancedInputLocalPlayerSubsystem* Subsystem = ULocalPlayer::GetSubsystem<UEnhancedInputLocalPlayerSubsystem>(PlayerController->GetLocalPlayer()))
-		{
-			Subsystem->RemoveMappingContext(DefaultMappingContext);
-		}
-	}
-}
 
 void ATP_FirstPersonCharacter::Move(const FInputActionValue& Value)
 {
diff --git a/Templates/TP_FirstPerson/Source/TP_FirstPerson/TP_FirstPersonCharacter.h b/Templates/TP_FirstPerson/Source/TP_FirstPerson/TP_FirstPersonCharacter.h
--- a/Templates/TP_FirstPerson/Source/TP_FirstPerson/TP_FirstPersonCharacter.h
+++ b/Templates/TP_FirstPerson/Source/TP_FirstPerson/TP_FirstPersonCharacter.h
@@ -47,11 +47,7 @@ protected:
 	virtual void BeginPlay();
 
 public:
-	
-	/** Called for adding mapping context*/
-	virtual void PossessedBy(AController* NewController) override;
-	virtual void UnPossessed() override;
-
+		
 	/** Look Input Action */
 	UPROPERTY(EditAnywhere, BlueprintReadOnly, Category = Input, meta = (AllowPrivateAccess = "true"))
 	class UInputAction* LookAction;
diff --git a/Templates/TP_FirstPerson/Source/TP_FirstPerson/TP_WeaponComponent.cpp b/Templates/TP_FirstPerson/Source/TP_FirstPerson/TP_WeaponComponent.cpp
--- a/Templates/TP_FirstPerson/Source/TP_FirstPerson/TP_WeaponComponent.cpp
+++ b/Templates/TP_FirstPerson/Source/TP_FirstPerson/TP_WeaponComponent.cpp
@@ -87,7 +87,7 @@ void UTP_WeaponComponent::AttachWeapon(ATP_FirstPersonCharacter* TargetCharacter
 			Subsystem->AddMappingContext(FireMappingContext, 1);
 		}
 
-		if (UEnhancedInputComponent* EnhancedInputComponent = CastChecked<UEnhancedInputComponent>(PlayerController->InputComponent))
+		if (UEnhancedInputComponent* EnhancedInputComponent = Cast<UEnhancedInputComponent>(PlayerController->InputComponent))
 		{
 			// Fire
 			EnhancedInputComponent->BindAction(FireAction, ETriggerEvent::Triggered, this, &UTP_WeaponComponent::Fire);
diff --git a/Templates/TP_ThirdPerson/Source/TP_ThirdPerson/TP_ThirdPersonCharacter.cpp b/Templates/TP_ThirdPerson/Source/TP_ThirdPerson/TP_ThirdPersonCharacter.cpp
--- a/Templates/TP_ThirdPerson/Source/TP_ThirdPerson/TP_ThirdPersonCharacter.cpp
+++ b/Templates/TP_ThirdPerson/Source/TP_ThirdPerson/TP_ThirdPersonCharacter.cpp
@@ -51,12 +51,13 @@ ATP_ThirdPersonCharacter::ATP_ThirdPersonCharacter()
 	// are set in the derived blueprint asset named ThirdPersonCharacter (to avoid direct content references in C++)
 }
 
-void ATP_ThirdPersonCharacter::PossessedBy(AController* NewController)
+void ATP_ThirdPersonCharacter::BeginPlay()
 {
-	Super::PossessedBy(NewController);
+	// Call the base class  
+	Super::BeginPlay();
 
-	//Add a mapping context on Possessed
-	if (APlayerController* PlayerController = Cast<APlayerController>(NewController))
+	//Add Input Mapping Context
+	if (APlayerController* PlayerController = Cast<APlayerController>(Controller))
 	{
 		if (UEnhancedInputLocalPlayerSubsystem* Subsystem = ULocalPlayer::GetSubsystem<UEnhancedInputLocalPlayerSubsystem>(PlayerController->GetLocalPlayer()))
 		{
@@ -65,19 +66,6 @@ void ATP_ThirdPersonCharacter::PossessedBy(AController* NewController)
 	}
 }
 
-void ATP_ThirdPersonCharacter::UnPossessed()
-{
-	Super::UnPossessed();
-
-	if (APlayerController* PlayerController = Cast<APlayerController>(GetController()))
-	{
-		if (UEnhancedInputLocalPlayerSubsystem* Subsystem = ULocalPlayer::GetSubsystem<UEnhancedInputLocalPlayerSubsystem>(PlayerController->GetLocalPlayer()))
-		{
-			Subsystem->RemoveMappingContext(DefaultMappingContext);
-		}
-	}
-}
-
 //////////////////////////////////////////////////////////////////////////
 // Input
 
diff --git a/Templates/TP_ThirdPerson/Source/TP_ThirdPerson/TP_ThirdPersonCharacter.h b/Templates/TP_ThirdPerson/Source/TP_ThirdPerson/TP_ThirdPersonCharacter.h
--- a/Templates/TP_ThirdPerson/Source/TP_ThirdPerson/TP_ThirdPersonCharacter.h
+++ b/Templates/TP_ThirdPerson/Source/TP_ThirdPerson/TP_ThirdPersonCharacter.h
@@ -39,10 +39,7 @@ class ATP_ThirdPersonCharacter : public ACharacter
 
 public:
 	ATP_ThirdPersonCharacter();
-
-	/** Called for adding mapping context*/
-	virtual void PossessedBy(AController* NewController) override;
-	virtual void UnPossessed() override;
+	
 
 protected:
 
@@ -56,7 +53,9 @@ protected:
 protected:
 	// APawn interface
 	virtual void SetupPlayerInputComponent(class UInputComponent* PlayerInputComponent) override;
-	// End of APawn interface
+	
+	// To add mapping context
+	virtual void BeginPlay();
 
 public:
 	/** Returns CameraBoom subobject **/
diff --git a/Templates/TP_TopDown/Source/TP_TopDown/TP_TopDownPlayerController.cpp b/Templates/TP_TopDown/Source/TP_TopDown/TP_TopDownPlayerController.cpp
--- a/Templates/TP_TopDown/Source/TP_TopDown/TP_TopDownPlayerController.cpp
+++ b/Templates/TP_TopDown/Source/TP_TopDown/TP_TopDownPlayerController.cpp
@@ -18,27 +18,18 @@ ATP_TopDownPlayerController::ATP_TopDownPlayerController()
 	FollowTime = 0.f;
 }
 
-void ATP_TopDownPlayerController::OnPossess(APawn* PossessedPawn)
+void ATP_TopDownPlayerController::BeginPlay()
 {
-	Super::OnPossess(PossessedPawn);
+	// Call the base class  
+	Super::BeginPlay();
 
+	//Add Input Mapping Context
 	if (UEnhancedInputLocalPlayerSubsystem* Subsystem = ULocalPlayer::GetSubsystem<UEnhancedInputLocalPlayerSubsystem>(GetLocalPlayer()))
 	{
 		Subsystem->AddMappingContext(DefaultMappingContext, 0);
 	}
 }
 
-void ATP_TopDownPlayerController::OnUnPossess()
-{
-	Super::OnUnPossess();
-
-	if (UEnhancedInputLocalPlayerSubsystem* Subsystem = ULocalPlayer::GetSubsystem<UEnhancedInputLocalPlayerSubsystem>(GetLocalPlayer()))
-	{
-		check(DefaultMappingContext);
-		Subsystem->RemoveMappingContext(DefaultMappingContext);
-	}
-}
-
 void ATP_TopDownPlayerController::SetupInputComponent()
 {
 	// set up gameplay key bindings
diff --git a/Templates/TP_TopDown/Source/TP_TopDown/TP_TopDownPlayerController.h b/Templates/TP_TopDown/Source/TP_TopDown/TP_TopDownPlayerController.h
--- a/Templates/TP_TopDown/Source/TP_TopDown/TP_TopDownPlayerController.h
+++ b/Templates/TP_TopDown/Source/TP_TopDown/TP_TopDownPlayerController.h
@@ -43,11 +43,10 @@ protected:
 	/** True if the controlled character should navigate to the mouse cursor. */
 	uint32 bMoveToMouseCursor : 1;
 
-	// Begin PlayerController interface
-	virtual void OnPossess(APawn* PossessedPawn) override;
-	virtual void OnUnPossess() override;
 	virtual void SetupInputComponent() override;
-	// End PlayerController interface
+	
+	// To add mapping context
+	virtual void BeginPlay();
 
 	/** Input handlers for SetDestination action. */
 	void OnInputStarted();
