Update all templates to BuildSettingsVersion.V3 & EngineIncludeOrderVersion.Unreal5_3

#rnx

diff --git a/Templates/TP_Blank/Source/TP_Blank.Target.cs b/Templates/TP_Blank/Source/TP_Blank.Target.cs
--- a/Templates/TP_Blank/Source/TP_Blank.Target.cs
+++ b/Templates/TP_Blank/Source/TP_Blank.Target.cs
@@ -8,8 +8,8 @@ public class TP_BlankTarget : TargetRules
 	public TP_BlankTarget(TargetInfo Target) : base(Target)
 	{
 		Type = TargetType.Game;
-		DefaultBuildSettings = BuildSettingsVersion.Latest;
-		IncludeOrderVersion = EngineIncludeOrderVersion.Unreal5_1;
+		DefaultBuildSettings = BuildSettingsVersion.V3;
+		IncludeOrderVersion = EngineIncludeOrderVersion.Unreal5_3;
 		ExtraModuleNames.Add("TP_Blank");
 	}
 }
diff --git a/Templates/TP_Blank/Source/TP_BlankEditor.Target.cs b/Templates/TP_Blank/Source/TP_BlankEditor.Target.cs
--- a/Templates/TP_Blank/Source/TP_BlankEditor.Target.cs
+++ b/Templates/TP_Blank/Source/TP_BlankEditor.Target.cs
@@ -8,8 +8,8 @@ public class TP_BlankEditorTarget : TargetRules
 	public TP_BlankEditorTarget( TargetInfo Target) : base(Target)
 	{
 		Type = TargetType.Editor;
-		DefaultBuildSettings = BuildSettingsVersion.Latest;
-		IncludeOrderVersion = EngineIncludeOrderVersion.Unreal5_1;
+		DefaultBuildSettings = BuildSettingsVersion.V3;
+		IncludeOrderVersion = EngineIncludeOrderVersion.Unreal5_3;
 		ExtraModuleNames.Add("TP_Blank");
 	}
 }
diff --git a/Templates/TP_FirstPerson/Source/TP_FirstPerson.Target.cs b/Templates/TP_FirstPerson/Source/TP_FirstPerson.Target.cs
--- a/Templates/TP_FirstPerson/Source/TP_FirstPerson.Target.cs
+++ b/Templates/TP_FirstPerson/Source/TP_FirstPerson.Target.cs
@@ -8,8 +8,8 @@ public class TP_FirstPersonTarget : TargetRules
 	public TP_FirstPersonTarget(TargetInfo Target) : base(Target)
 	{
 		Type = TargetType.Game;
-		DefaultBuildSettings = BuildSettingsVersion.Latest;
-		IncludeOrderVersion = EngineIncludeOrderVersion.Unreal5_1;
+		DefaultBuildSettings = BuildSettingsVersion.V3;
+		IncludeOrderVersion = EngineIncludeOrderVersion.Unreal5_3;
 		ExtraModuleNames.Add("TP_FirstPerson");
 	}
 }
diff --git a/Templates/TP_FirstPerson/Source/TP_FirstPersonEditor.Target.cs b/Templates/TP_FirstPerson/Source/TP_FirstPersonEditor.Target.cs
--- a/Templates/TP_FirstPerson/Source/TP_FirstPersonEditor.Target.cs
+++ b/Templates/TP_FirstPerson/Source/TP_FirstPersonEditor.Target.cs
@@ -8,8 +8,8 @@ public class TP_FirstPersonEditorTarget : TargetRules
 	public TP_FirstPersonEditorTarget(TargetInfo Target) : base(Target)
 	{
 		Type = TargetType.Editor;
-		DefaultBuildSettings = BuildSettingsVersion.Latest;
-		IncludeOrderVersion = EngineIncludeOrderVersion.Unreal5_1;
+		DefaultBuildSettings = BuildSettingsVersion.V3;
+		IncludeOrderVersion = EngineIncludeOrderVersion.Unreal5_3;
 		ExtraModuleNames.Add("TP_FirstPerson");
 	}
 }
diff --git a/Templates/TP_Puzzle/Source/TP_Puzzle.Target.cs b/Templates/TP_Puzzle/Source/TP_Puzzle.Target.cs
--- a/Templates/TP_Puzzle/Source/TP_Puzzle.Target.cs
+++ b/Templates/TP_Puzzle/Source/TP_Puzzle.Target.cs
@@ -8,8 +8,8 @@ public class TP_PuzzleTarget : TargetRules
     public TP_PuzzleTarget(TargetInfo Target) : base(Target)
 	{
 		Type = TargetType.Game;
-		DefaultBuildSettings = BuildSettingsVersion.Latest;
-		IncludeOrderVersion = EngineIncludeOrderVersion.Unreal5_1;
+		DefaultBuildSettings = BuildSettingsVersion.V3;
+		IncludeOrderVersion = EngineIncludeOrderVersion.Unreal5_3;
 		ExtraModuleNames.Add("TP_Puzzle");
 	}
 }
diff --git a/Templates/TP_Puzzle/Source/TP_PuzzleEditor.Target.cs b/Templates/TP_Puzzle/Source/TP_PuzzleEditor.Target.cs
--- a/Templates/TP_Puzzle/Source/TP_PuzzleEditor.Target.cs
+++ b/Templates/TP_Puzzle/Source/TP_PuzzleEditor.Target.cs
@@ -8,8 +8,8 @@ public class TP_PuzzleEditorTarget : TargetRules
     public TP_PuzzleEditorTarget(TargetInfo Target) : base(Target)
 	{
 		Type = TargetType.Editor;
-		DefaultBuildSettings = BuildSettingsVersion.Latest;
-		IncludeOrderVersion = EngineIncludeOrderVersion.Unreal5_1;
+		DefaultBuildSettings = BuildSettingsVersion.V3;
+		IncludeOrderVersion = EngineIncludeOrderVersion.Unreal5_3;
 		ExtraModuleNames.Add("TP_Puzzle");
 	}
 }
diff --git a/Templates/TP_SIM_Blank/Source/TP_SIM_Blank.Target.cs b/Templates/TP_SIM_Blank/Source/TP_SIM_Blank.Target.cs
--- a/Templates/TP_SIM_Blank/Source/TP_SIM_Blank.Target.cs
+++ b/Templates/TP_SIM_Blank/Source/TP_SIM_Blank.Target.cs
@@ -8,7 +8,8 @@ public class TP_SIM_BlankTarget : TargetRules
 	public TP_SIM_BlankTarget( TargetInfo Target) : base(Target)
 	{
 		Type = TargetType.Game;
-		DefaultBuildSettings = BuildSettingsVersion.V2;
+		DefaultBuildSettings = BuildSettingsVersion.V3;
+		IncludeOrderVersion = EngineIncludeOrderVersion.Unreal5_3;
 		ExtraModuleNames.AddRange( new string[] { "TP_SIM_Blank" } );
 	}
 }
diff --git a/Templates/TP_SIM_Blank/Source/TP_SIM_BlankEditor.Target.cs b/Templates/TP_SIM_Blank/Source/TP_SIM_BlankEditor.Target.cs
--- a/Templates/TP_SIM_Blank/Source/TP_SIM_BlankEditor.Target.cs
+++ b/Templates/TP_SIM_Blank/Source/TP_SIM_BlankEditor.Target.cs
@@ -8,7 +8,8 @@ public class TP_SIM_BlankEditorTarget : TargetRules
 	public TP_SIM_BlankEditorTarget( TargetInfo Target) : base(Target)
 	{
 		Type = TargetType.Editor;
-		DefaultBuildSettings = BuildSettingsVersion.V2;
+		DefaultBuildSettings = BuildSettingsVersion.V3;
+		IncludeOrderVersion = EngineIncludeOrderVersion.Unreal5_3;
 		ExtraModuleNames.AddRange( new string[] { "TP_SIM_Blank" } );
 	}
 }
diff --git a/Templates/TP_SIM_Sandbox/Source/TP_SIM_Sandbox.Target.cs b/Templates/TP_SIM_Sandbox/Source/TP_SIM_Sandbox.Target.cs
--- a/Templates/TP_SIM_Sandbox/Source/TP_SIM_Sandbox.Target.cs
+++ b/Templates/TP_SIM_Sandbox/Source/TP_SIM_Sandbox.Target.cs
@@ -8,7 +8,8 @@ public class TP_SIM_SandboxTarget : TargetRules
 	public TP_SIM_SandboxTarget( TargetInfo Target) : base(Target)
 	{
 		Type = TargetType.Game;
-		DefaultBuildSettings = BuildSettingsVersion.V2;
+		DefaultBuildSettings = BuildSettingsVersion.V3;
+		IncludeOrderVersion = EngineIncludeOrderVersion.Unreal5_3;
 		ExtraModuleNames.AddRange( new string[] { "TP_SIM_Sandbox" } );
 	}
 }
diff --git a/Templates/TP_SIM_Sandbox/Source/TP_SIM_SandboxEditor.Target.cs b/Templates/TP_SIM_Sandbox/Source/TP_SIM_SandboxEditor.Target.cs
--- a/Templates/TP_SIM_Sandbox/Source/TP_SIM_SandboxEditor.Target.cs
+++ b/Templates/TP_SIM_Sandbox/Source/TP_SIM_SandboxEditor.Target.cs
@@ -8,7 +8,8 @@ public class TP_SIM_SandboxEditorTarget : TargetRules
 	public TP_SIM_SandboxEditorTarget( TargetInfo Target) : base(Target)
 	{
 		Type = TargetType.Editor;
-		DefaultBuildSettings = BuildSettingsVersion.V2;
+		DefaultBuildSettings = BuildSettingsVersion.V3;
+		IncludeOrderVersion = EngineIncludeOrderVersion.Unreal5_3;
 		ExtraModuleNames.AddRange( new string[] { "TP_SIM_Sandbox" } );
 	}
 }
diff --git a/Templates/TP_ThirdPerson/Source/TP_ThirdPerson.Target.cs b/Templates/TP_ThirdPerson/Source/TP_ThirdPerson.Target.cs
--- a/Templates/TP_ThirdPerson/Source/TP_ThirdPerson.Target.cs
+++ b/Templates/TP_ThirdPerson/Source/TP_ThirdPerson.Target.cs
@@ -8,8 +8,8 @@ public class TP_ThirdPersonTarget : TargetRules
 	public TP_ThirdPersonTarget(TargetInfo Target) : base(Target)
 	{
 		Type = TargetType.Game;
-		DefaultBuildSettings = BuildSettingsVersion.Latest;
-		IncludeOrderVersion = EngineIncludeOrderVersion.Unreal5_1;
+		DefaultBuildSettings = BuildSettingsVersion.V3;
+		IncludeOrderVersion = EngineIncludeOrderVersion.Unreal5_3;
 		ExtraModuleNames.Add("TP_ThirdPerson");
 	}
 }
diff --git a/Templates/TP_ThirdPerson/Source/TP_ThirdPersonEditor.Target.cs b/Templates/TP_ThirdPerson/Source/TP_ThirdPersonEditor.Target.cs
--- a/Templates/TP_ThirdPerson/Source/TP_ThirdPersonEditor.Target.cs
+++ b/Templates/TP_ThirdPerson/Source/TP_ThirdPersonEditor.Target.cs
@@ -8,8 +8,8 @@ public class TP_ThirdPersonEditorTarget : TargetRules
 	public TP_ThirdPersonEditorTarget(TargetInfo Target) : base(Target)
 	{
 		Type = TargetType.Editor;
-		DefaultBuildSettings = BuildSettingsVersion.Latest;
-		IncludeOrderVersion = EngineIncludeOrderVersion.Unreal5_1;
+		DefaultBuildSettings = BuildSettingsVersion.V3;
+		IncludeOrderVersion = EngineIncludeOrderVersion.Unreal5_3;
 		ExtraModuleNames.Add("TP_ThirdPerson");
 	}
 }
diff --git a/Templates/TP_TopDown/Source/TP_TopDown.Target.cs b/Templates/TP_TopDown/Source/TP_TopDown.Target.cs
--- a/Templates/TP_TopDown/Source/TP_TopDown.Target.cs
+++ b/Templates/TP_TopDown/Source/TP_TopDown.Target.cs
@@ -8,8 +8,8 @@ public class TP_TopDownTarget : TargetRules
 	public TP_TopDownTarget(TargetInfo Target) : base(Target)
 	{
 		Type = TargetType.Game;
-		DefaultBuildSettings = BuildSettingsVersion.Latest;
-		IncludeOrderVersion = EngineIncludeOrderVersion.Unreal5_1;
+		DefaultBuildSettings = BuildSettingsVersion.V3;
+		IncludeOrderVersion = EngineIncludeOrderVersion.Unreal5_3;
 		ExtraModuleNames.Add("TP_TopDown");
 	}
 }
diff --git a/Templates/TP_TopDown/Source/TP_TopDownEditor.Target.cs b/Templates/TP_TopDown/Source/TP_TopDownEditor.Target.cs
--- a/Templates/TP_TopDown/Source/TP_TopDownEditor.Target.cs
+++ b/Templates/TP_TopDown/Source/TP_TopDownEditor.Target.cs
@@ -8,8 +8,8 @@ public class TP_TopDownEditorTarget : TargetRules
 	public TP_TopDownEditorTarget(TargetInfo Target) : base(Target)
 	{
 		Type = TargetType.Editor;
-		DefaultBuildSettings = BuildSettingsVersion.Latest;
-		IncludeOrderVersion = EngineIncludeOrderVersion.Unreal5_1;
+		DefaultBuildSettings = BuildSettingsVersion.V3;
+		IncludeOrderVersion = EngineIncludeOrderVersion.Unreal5_3;
 		ExtraModuleNames.Add("TP_TopDown");
 	}
 }
diff --git a/Templates/TP_VehicleAdv/Source/TP_VehicleAdv.Target.cs b/Templates/TP_VehicleAdv/Source/TP_VehicleAdv.Target.cs
--- a/Templates/TP_VehicleAdv/Source/TP_VehicleAdv.Target.cs
+++ b/Templates/TP_VehicleAdv/Source/TP_VehicleAdv.Target.cs
@@ -8,8 +8,8 @@ public class TP_VehicleAdvTarget : TargetRules
 	public TP_VehicleAdvTarget(TargetInfo Target) : base(Target)
 	{
 		Type = TargetType.Game;
-		DefaultBuildSettings = BuildSettingsVersion.Latest;
-		IncludeOrderVersion = EngineIncludeOrderVersion.Unreal5_1;
+		DefaultBuildSettings = BuildSettingsVersion.V3;
+		IncludeOrderVersion = EngineIncludeOrderVersion.Unreal5_3;
 		ExtraModuleNames.Add("TP_VehicleAdv");
 	}
 }
diff --git a/Templates/TP_VehicleAdv/Source/TP_VehicleAdvEditor.Target.cs b/Templates/TP_VehicleAdv/Source/TP_VehicleAdvEditor.Target.cs
--- a/Templates/TP_VehicleAdv/Source/TP_VehicleAdvEditor.Target.cs
+++ b/Templates/TP_VehicleAdv/Source/TP_VehicleAdvEditor.Target.cs
@@ -8,8 +8,8 @@ public class TP_VehicleAdvEditorTarget : TargetRules
 	public TP_VehicleAdvEditorTarget(TargetInfo Target) : base(Target)
 	{
 		Type = TargetType.Editor;
-		DefaultBuildSettings = BuildSettingsVersion.Latest;
-		IncludeOrderVersion = EngineIncludeOrderVersion.Unreal5_1;
+		DefaultBuildSettings = BuildSettingsVersion.V3;
+		IncludeOrderVersion = EngineIncludeOrderVersion.Unreal5_3;
 		ExtraModuleNames.Add("TP_VehicleAdv");
 	}
 }
