Multiple weapon pickups no longer stack additional weapons on the character. This fix will only work in single player; currently evaluating a refactor for multiplayer support.
Also standardizing jump input and controls following third person fix in UE-173720

diff --git a/Templates/TP_FirstPerson/Source/TP_FirstPerson/TP_FirstPersonCharacter.cpp b/Templates/TP_FirstPerson/Source/TP_FirstPerson/TP_FirstPersonCharacter.cpp
--- a/Templates/TP_FirstPerson/Source/TP_FirstPerson/TP_FirstPersonCharacter.cpp
+++ b/Templates/TP_FirstPerson/Source/TP_FirstPerson/TP_FirstPersonCharacter.cpp
@@ -65,7 +65,7 @@ void ATP_FirstPersonCharacter::SetupPlayerInputComponent(UInputComponent* Player
 	if (UEnhancedInputComponent* EnhancedInputComponent = Cast<UEnhancedInputComponent>(PlayerInputComponent))
 	{
 		// Jumping
-		EnhancedInputComponent->BindAction(JumpAction, ETriggerEvent::Triggered, this, &ACharacter::Jump);
+		EnhancedInputComponent->BindAction(JumpAction, ETriggerEvent::Started, this, &ACharacter::Jump);
 		EnhancedInputComponent->BindAction(JumpAction, ETriggerEvent::Completed, this, &ACharacter::StopJumping);
 
 		// Moving
diff --git a/Templates/TP_FirstPerson/Source/TP_FirstPerson/TP_WeaponComponent.cpp b/Templates/TP_FirstPerson/Source/TP_FirstPerson/TP_WeaponComponent.cpp
--- a/Templates/TP_FirstPerson/Source/TP_FirstPerson/TP_WeaponComponent.cpp
+++ b/Templates/TP_FirstPerson/Source/TP_FirstPerson/TP_WeaponComponent.cpp
@@ -66,7 +66,9 @@ void UTP_WeaponComponent::Fire()
 void UTP_WeaponComponent::AttachWeapon(ATP_FirstPersonCharacter* TargetCharacter)
 {
 	Character = TargetCharacter;
-	if (Character == nullptr)
+
+	// Check that the character is valid, and has no rifle yet
+	if (Character == nullptr || Character->GetHasRifle())
 	{
 		return;
 	}
