PR #10784: Fixed First Person template project not compiling when PCH is disabled.

diff --git a/Templates/TP_FirstPerson/Source/TP_FirstPerson/TP_WeaponComponent.cpp b/Templates/TP_FirstPerson/Source/TP_FirstPerson/TP_WeaponComponent.cpp
--- a/Templates/TP_FirstPerson/Source/TP_FirstPerson/TP_WeaponComponent.cpp
+++ b/Templates/TP_FirstPerson/Source/TP_FirstPerson/TP_WeaponComponent.cpp
@@ -9,6 +9,9 @@
 #include "Kismet/GameplayStatics.h"
 #include "EnhancedInputComponent.h"
 #include "EnhancedInputSubsystems.h"
+#include "Animation/AnimInstance.h"
+#include "Engine/LocalPlayer.h"
+#include "Engine/World.h"
 
 // Sets default values for this component's properties
 UTP_WeaponComponent::UTP_WeaponComponent()
