Templates - Removed references to weapon from Characters on TP_FirstPerson and TP_FirstPersonBP

diff --git a/Templates/TP_FirstPerson/Source/TP_FirstPerson/TP_FirstPersonCharacter.cpp b/Templates/TP_FirstPerson/Source/TP_FirstPerson/TP_FirstPersonCharacter.cpp
--- a/Templates/TP_FirstPerson/Source/TP_FirstPerson/TP_FirstPersonCharacter.cpp
+++ b/Templates/TP_FirstPerson/Source/TP_FirstPerson/TP_FirstPersonCharacter.cpp
@@ -18,9 +18,6 @@ DEFINE_LOG_CATEGORY(LogTemplateCharacter);
 
 ATP_FirstPersonCharacter::ATP_FirstPersonCharacter()
 {
-	// Character doesnt have a rifle at start
-	bHasRifle = false;
-	
 	// Set size for collision capsule
 	GetCapsuleComponent()->InitCapsuleSize(55.f, 96.0f);
 		
@@ -50,20 +47,7 @@ void ATP_FirstPersonCharacter::BeginPlay()
 //////////////////////////////////////////////////////////////////////////// Input
 
 void ATP_FirstPersonCharacter::SetupPlayerInputComponent(UInputComponent* PlayerInputComponent)
-{
-	// Add Input Mapping Context
-	if (APlayerController* PlayerController = Cast<APlayerController>(Controller))
-	{
-		if (UEnhancedInputLocalPlayerSubsystem* Subsystem = ULocalPlayer::GetSubsystem<UEnhancedInputLocalPlayerSubsystem>(PlayerController->GetLocalPlayer()))
-		{
-			Subsystem->AddMappingContext(DefaultMappingContext, 0);
-		}
-	}
-	else
-	{
-		UE_LOG(LogTemplateCharacter, Error, TEXT("'%s' Failed to add an input mapping context because there is no owning Player Controller!"), *GetNameSafe(this));
-	}
-	
+{	
 	// Set up action bindings
 	if (UEnhancedInputComponent* EnhancedInputComponent = Cast<UEnhancedInputComponent>(PlayerInputComponent))
 	{
@@ -108,14 +92,4 @@ void ATP_FirstPersonCharacter::Look(const FInputActionValue& Value)
 		AddControllerYawInput(LookAxisVector.X);
 		AddControllerPitchInput(LookAxisVector.Y);
 	}
-}
-
-void ATP_FirstPersonCharacter::SetHasRifle(bool bNewHasRifle)
-{
-	bHasRifle = bNewHasRifle;
-}
-
-bool ATP_FirstPersonCharacter::GetHasRifle()
-{
-	return bHasRifle;
 }
\ No newline at end of file
diff --git a/Templates/TP_FirstPerson/Source/TP_FirstPerson/TP_FirstPersonCharacter.h b/Templates/TP_FirstPerson/Source/TP_FirstPerson/TP_FirstPersonCharacter.h
--- a/Templates/TP_FirstPerson/Source/TP_FirstPerson/TP_FirstPersonCharacter.h
+++ b/Templates/TP_FirstPerson/Source/TP_FirstPerson/TP_FirstPersonCharacter.h
@@ -22,17 +22,13 @@ class ATP_FirstPersonCharacter : public ACharacter
 	GENERATED_BODY()
 
 	/** Pawn mesh: 1st person view (arms; seen only by self) */
-	UPROPERTY(VisibleDefaultsOnly, Category=Mesh)
+	UPROPERTY(VisibleAnywhere, BlueprintReadOnly, Category=Mesh, meta = (AllowPrivateAccess = "true"))
 	USkeletalMeshComponent* Mesh1P;
 
 	/** First person camera */
 	UPROPERTY(VisibleAnywhere, BlueprintReadOnly, Category = Camera, meta = (AllowPrivateAccess = "true"))
 	UCameraComponent* FirstPersonCameraComponent;
 
-	/** MappingContext */
-	UPROPERTY(EditAnywhere, BlueprintReadOnly, Category=Input, meta=(AllowPrivateAccess = "true"))
-	UInputMappingContext* DefaultMappingContext;
-
 	/** Jump Input Action */
 	UPROPERTY(EditAnywhere, BlueprintReadOnly, Category=Input, meta=(AllowPrivateAccess = "true"))
 	UInputAction* JumpAction;
@@ -53,18 +49,6 @@ public:
 	UPROPERTY(EditAnywhere, BlueprintReadOnly, Category = Input, meta = (AllowPrivateAccess = "true"))
 	class UInputAction* LookAction;
 
-	/** Bool for AnimBP to switch to another animation set */
-	UPROPERTY(VisibleAnywhere, BlueprintReadOnly, Category = Weapon)
-	bool bHasRifle;
-
-	/** Setter to set the bool */
-	UFUNCTION(BlueprintCallable, Category = Weapon)
-	void SetHasRifle(bool bNewHasRifle);
-
-	/** Getter for the bool */
-	UFUNCTION(BlueprintCallable, Category = Weapon)
-	bool GetHasRifle();
-
 protected:
 	/** Called for movement input */
 	void Move(const FInputActionValue& Value);
diff --git a/Templates/TP_FirstPerson/Source/TP_FirstPerson/TP_FirstPersonPlayerController.cpp b/Templates/TP_FirstPerson/Source/TP_FirstPerson/TP_FirstPersonPlayerController.cpp
--- a/Templates/TP_FirstPerson/Source/TP_FirstPerson/TP_FirstPersonPlayerController.cpp
+++ b/Templates/TP_FirstPerson/Source/TP_FirstPerson/TP_FirstPersonPlayerController.cpp
@@ -3,6 +3,7 @@
 
 #include "TP_FirstPersonPlayerController.h"
 #include "EnhancedInputSubsystems.h"
+#include "Engine/LocalPlayer.h"
 
 void ATP_FirstPersonPlayerController::BeginPlay()
 {
@@ -13,7 +14,5 @@ void ATP_FirstPersonPlayerController::BeginPlay()
 	{
 		// add the mapping context so we get controls
 		Subsystem->AddMappingContext(InputMappingContext, 0);
-
-		UE_LOG(LogTemp, Warning, TEXT("BeginPlay"));
 	}
 }
\ No newline at end of file
diff --git a/Templates/TP_FirstPerson/Source/TP_FirstPerson/TP_WeaponComponent.cpp b/Templates/TP_FirstPerson/Source/TP_FirstPerson/TP_WeaponComponent.cpp
--- a/Templates/TP_FirstPerson/Source/TP_FirstPerson/TP_WeaponComponent.cpp
+++ b/Templates/TP_FirstPerson/Source/TP_FirstPerson/TP_WeaponComponent.cpp
@@ -66,22 +66,22 @@ void UTP_WeaponComponent::Fire()
 	}
 }
 
-void UTP_WeaponComponent::AttachWeapon(ATP_FirstPersonCharacter* TargetCharacter)
+bool UTP_WeaponComponent::AttachWeapon(ATP_FirstPersonCharacter* TargetCharacter)
 {
 	Character = TargetCharacter;
 
-	// Check that the character is valid, and has no rifle yet
-	if (Character == nullptr || Character->GetHasRifle())
+	// Check that the character is valid, and has no weapon component yet
+	if (Character == nullptr || Character->GetInstanceComponents().FindItemByClass<UTP_WeaponComponent>())
 	{
-		return;
+		return false;
 	}
 
 	// Attach the weapon to the First Person Character
 	FAttachmentTransformRules AttachmentRules(EAttachmentRule::SnapToTarget, true);
 	AttachToComponent(Character->GetMesh1P(), AttachmentRules, FName(TEXT("GripPoint")));
-	
-	// switch bHasRifle so the animation blueprint can switch to another animation set
-	Character->SetHasRifle(true);
+
+	// add the weapon as an instance component to the character
+	Character->AddInstanceComponent(this);
 
 	// Set up action bindings
 	if (APlayerController* PlayerController = Cast<APlayerController>(Character->GetController()))
@@ -98,6 +98,8 @@ void UTP_WeaponComponent::AttachWeapon(ATP_FirstPersonCharacter* TargetCharacter
 			EnhancedInputComponent->BindAction(FireAction, ETriggerEvent::Triggered, this, &UTP_WeaponComponent::Fire);
 		}
 	}
+
+	return true;
 }
 
 void UTP_WeaponComponent::EndPlay(const EEndPlayReason::Type EndPlayReason)
diff --git a/Templates/TP_FirstPerson/Source/TP_FirstPerson/TP_WeaponComponent.h b/Templates/TP_FirstPerson/Source/TP_FirstPerson/TP_WeaponComponent.h
--- a/Templates/TP_FirstPerson/Source/TP_FirstPerson/TP_WeaponComponent.h
+++ b/Templates/TP_FirstPerson/Source/TP_FirstPerson/TP_WeaponComponent.h
@@ -43,7 +43,7 @@ public:
 
 	/** Attaches the actor to a FirstPersonCharacter */
 	UFUNCTION(BlueprintCallable, Category="Weapon")
-	void AttachWeapon(ATP_FirstPersonCharacter* TargetCharacter);
+	bool AttachWeapon(ATP_FirstPersonCharacter* TargetCharacter);
 
 	/** Make the weapon Fire a Projectile */
 	UFUNCTION(BlueprintCallable, Category="Weapon")
