TP_FirstPerson - added missing EndPlay call to WeaponComponent.
Clarified comments on instanced component attachment of WeaponComponent.

diff --git a/Templates/TP_FirstPerson/Source/TP_FirstPerson/TP_WeaponComponent.cpp b/Templates/TP_FirstPerson/Source/TP_FirstPerson/TP_WeaponComponent.cpp
--- a/Templates/TP_FirstPerson/Source/TP_FirstPerson/TP_WeaponComponent.cpp
+++ b/Templates/TP_FirstPerson/Source/TP_FirstPerson/TP_WeaponComponent.cpp
@@ -80,7 +80,8 @@ bool UTP_WeaponComponent::AttachWeapon(ATP_FirstPersonCharacter* TargetCharacter
 	FAttachmentTransformRules AttachmentRules(EAttachmentRule::SnapToTarget, true);
 	AttachToComponent(Character->GetMesh1P(), AttachmentRules, FName(TEXT("GripPoint")));
 
-	// add the weapon as an instance component to the character
+	// the weapon component is owned by the Pickup, so we need to 
+	// manually add it to the Character as an Instance Component to hand off ownership
 	Character->AddInstanceComponent(this);
 
 	// Set up action bindings
@@ -104,16 +105,19 @@ bool UTP_WeaponComponent::AttachWeapon(ATP_FirstPersonCharacter* TargetCharacter
 
 void UTP_WeaponComponent::EndPlay(const EEndPlayReason::Type EndPlayReason)
 {
-	if (Character == nullptr)
+	// ensure we have a character owner
+	if (Character != nullptr)
 	{
-		return;
-	}
-
-	if (APlayerController* PlayerController = Cast<APlayerController>(Character->GetController()))
-	{
-		if (UEnhancedInputLocalPlayerSubsystem* Subsystem = ULocalPlayer::GetSubsystem<UEnhancedInputLocalPlayerSubsystem>(PlayerController->GetLocalPlayer()))
+		// remove the input mapping context from the Player Controller
+		if (APlayerController* PlayerController = Cast<APlayerController>(Character->GetController()))
 		{
-			Subsystem->RemoveMappingContext(FireMappingContext);
+			if (UEnhancedInputLocalPlayerSubsystem* Subsystem = ULocalPlayer::GetSubsystem<UEnhancedInputLocalPlayerSubsystem>(PlayerController->GetLocalPlayer()))
+			{
+				Subsystem->RemoveMappingContext(FireMappingContext);
+			}
 		}
 	}
+
+	// maintain the EndPlay call chain
+	Super::EndPlay(EndPlayReason);
 }
\ No newline at end of file
