Templates - Updated Feature Packs for game templates and Starter Content
Modified C++ Feature Pack configs so that all source files are picked up by default.
Renamed First Person C++ Template source files that did not conform to automatic rename scheme and updated BP classes and actors that depended on these.

#rnx

diff --git a/Templates/TP_FirstPerson/Config/DefaultEngine.ini b/Templates/TP_FirstPerson/Config/DefaultEngine.ini
--- a/Templates/TP_FirstPerson/Config/DefaultEngine.ini
+++ b/Templates/TP_FirstPerson/Config/DefaultEngine.ini
@@ -57,3 +57,17 @@ SoundCueCookQualityIndex=-1
 -TargetedRHIs=SF_VULKAN_SM5
 +TargetedRHIs=SF_VULKAN_SM6
 
+[/Script/AndroidFileServerEditor.AndroidFileServerRuntimeSettings]
+bEnablePlugin=True
+bAllowNetworkConnection=True
+SecurityToken=7B0AA718435214EB1E253DACC60399C4
+bIncludeInShipping=False
+bAllowExternalStartInShipping=False
+bCompileAFSProject=False
+bUseCompression=False
+bLogFiles=False
+bReportStats=False
+ConnectionType=USBOnly
+bUseManualIPAddress=False
+ManualIPAddress=
+
diff --git a/Templates/TP_FirstPerson/Config/config.ini b/Templates/TP_FirstPerson/Config/config.ini
--- a/Templates/TP_FirstPerson/Config/config.ini
+++ b/Templates/TP_FirstPerson/Config/config.ini
@@ -23,11 +23,5 @@ GameName=TP_FirstPerson
 
 [AdditionalFilesToAdd]
 +Files=Templates/TP_FirstPerson/Content/*.*
-+Files=Templates/TP_FirstPerson/Source/TP_FirstPerson/TP_FirstPerson.h
-+Files=Templates/TP_FirstPerson/Source/TP_FirstPerson/TP_FirstPerson.cpp
-+Files=Templates/TP_FirstPerson/Source/TP_FirstPerson/TP_FirstPersonCharacter.cpp
-+Files=Templates/TP_FirstPerson/Source/TP_FirstPerson/TP_FirstPersonCharacter.h
-+Files=Templates/TP_FirstPerson/Source/TP_FirstPerson/TP_FirstPersonGameMode.cpp
-+Files=Templates/TP_FirstPerson/Source/TP_FirstPerson/TP_FirstPersonGameMode.h
-+Files=Templates/TP_FirstPerson/Source/TP_FirstPerson/TP_FirstPersonProjectile.cpp
-+Files=Templates/TP_FirstPerson/Source/TP_FirstPerson/TP_FirstPersonProjectile.h
++Files=Templates/TP_FirstPerson/Source/TP_FirstPerson/*.h
++Files=Templates/TP_FirstPerson/Source/TP_FirstPerson/*.cpp
\ No newline at end of file
diff --git a/Templates/TP_FirstPerson/Source/TP_FirstPerson/TP_FirstPersonPickUpComponent.cpp b/Templates/TP_FirstPerson/Source/TP_FirstPerson/TP_FirstPersonPickUpComponent.cpp
new file mode 100644
--- /dev/null
+++ b/Templates/TP_FirstPerson/Source/TP_FirstPerson/TP_FirstPersonPickUpComponent.cpp
@@ -0,0 +1,31 @@
+// Copyright Epic Games, Inc. All Rights Reserved.
+
+#include "TP_FirstPersonPickUpComponent.h"
+
+UTP_FirstPersonPickUpComponent::UTP_FirstPersonPickUpComponent()
+{
+	// Setup the Sphere Collision
+	SphereRadius = 32.f;
+}
+
+void UTP_FirstPersonPickUpComponent::BeginPlay()
+{
+	Super::BeginPlay();
+
+	// Register our Overlap Event
+	OnComponentBeginOverlap.AddDynamic(this, &UTP_FirstPersonPickUpComponent::OnSphereBeginOverlap);
+}
+
+void UTP_FirstPersonPickUpComponent::OnSphereBeginOverlap(UPrimitiveComponent* OverlappedComponent, AActor* OtherActor, UPrimitiveComponent* OtherComp, int32 OtherBodyIndex, bool bFromSweep, const FHitResult& SweepResult)
+{
+	// Checking if it is a First Person Character overlapping
+	ATP_FirstPersonCharacter* Character = Cast<ATP_FirstPersonCharacter>(OtherActor);
+	if(Character != nullptr)
+	{
+		// Notify that the actor is being picked up
+		OnPickUp.Broadcast(Character);
+
+		// Unregister from the Overlap Event so it is no longer triggered
+		OnComponentBeginOverlap.RemoveAll(this);
+	}
+}
diff --git a/Templates/TP_FirstPerson/Source/TP_FirstPerson/TP_FirstPersonPickUpComponent.h b/Templates/TP_FirstPerson/Source/TP_FirstPerson/TP_FirstPersonPickUpComponent.h
new file mode 100644
--- /dev/null
+++ b/Templates/TP_FirstPerson/Source/TP_FirstPerson/TP_FirstPersonPickUpComponent.h
@@ -0,0 +1,34 @@
+// Copyright Epic Games, Inc. All Rights Reserved.
+
+#pragma once
+
+#include "CoreMinimal.h"
+#include "Components/SphereComponent.h"
+#include "TP_FirstPersonCharacter.h"
+#include "TP_FirstPersonPickUpComponent.generated.h"
+
+// Declaration of the delegate that will be called when someone picks this up
+// The character picking this up is the parameter sent with the notification
+DECLARE_DYNAMIC_MULTICAST_DELEGATE_OneParam(FOnPickUp, ATP_FirstPersonCharacter*, PickUpCharacter);
+
+UCLASS(Blueprintable, BlueprintType, ClassGroup = (Custom), meta = (BlueprintSpawnableComponent))
+class TP_FIRSTPERSON_API UTP_FirstPersonPickUpComponent : public USphereComponent
+{
+	GENERATED_BODY()
+
+public:
+	
+	/** Delegate to whom anyone can subscribe to receive this event */
+	UPROPERTY(BlueprintAssignable, Category = "Interaction")
+	FOnPickUp OnPickUp;
+
+	UTP_FirstPersonPickUpComponent();
+protected:
+
+	/** Called when the game starts */
+	virtual void BeginPlay() override;
+
+	/** Code for when something overlaps this component */
+	UFUNCTION()
+	void OnSphereBeginOverlap(UPrimitiveComponent* OverlappedComponent, AActor* OtherActor, UPrimitiveComponent* OtherComp, int32 OtherBodyIndex, bool bFromSweep, const FHitResult& SweepResult);
+};
\ No newline at end of file
diff --git a/Templates/TP_FirstPerson/Source/TP_FirstPerson/TP_FirstPersonWeaponComponent.cpp b/Templates/TP_FirstPerson/Source/TP_FirstPerson/TP_FirstPersonWeaponComponent.cpp
new file mode 100644
--- /dev/null
+++ b/Templates/TP_FirstPerson/Source/TP_FirstPerson/TP_FirstPersonWeaponComponent.cpp
@@ -0,0 +1,119 @@
+// Copyright Epic Games, Inc. All Rights Reserved.
+
+
+#include "TP_FirstPersonWeaponComponent.h"
+#include "TP_FirstPersonCharacter.h"
+#include "TP_FirstPersonProjectile.h"
+#include "GameFramework/PlayerController.h"
+#include "Camera/PlayerCameraManager.h"
+#include "Kismet/GameplayStatics.h"
+#include "EnhancedInputComponent.h"
+#include "EnhancedInputSubsystems.h"
+#include "Animation/AnimInstance.h"
+#include "Engine/LocalPlayer.h"
+#include "Engine/World.h"
+
+// Sets default values for this component's properties
+UTP_FirstPersonWeaponComponent::UTP_FirstPersonWeaponComponent()
+{
+	// Default offset from the character location for projectiles to spawn
+	MuzzleOffset = FVector(100.0f, 0.0f, 10.0f);
+}
+
+
+void UTP_FirstPersonWeaponComponent::Fire()
+{
+	if (Character == nullptr || Character->GetController() == nullptr)
+	{
+		return;
+	}
+
+	// Try and fire a projectile
+	if (ProjectileClass != nullptr)
+	{
+		UWorld* const World = GetWorld();
+		if (World != nullptr)
+		{
+			APlayerController* PlayerController = Cast<APlayerController>(Character->GetController());
+			const FRotator SpawnRotation = PlayerController->PlayerCameraManager->GetCameraRotation();
+			// MuzzleOffset is in camera space, so transform it to world space before offsetting from the character location to find the final muzzle position
+			const FVector SpawnLocation = GetOwner()->GetActorLocation() + SpawnRotation.RotateVector(MuzzleOffset);
+	
+			//Set Spawn Collision Handling Override
+			FActorSpawnParameters ActorSpawnParams;
+			ActorSpawnParams.SpawnCollisionHandlingOverride = ESpawnActorCollisionHandlingMethod::AdjustIfPossibleButDontSpawnIfColliding;
+	
+			// Spawn the projectile at the muzzle
+			World->SpawnActor<ATP_FirstPersonProjectile>(ProjectileClass, SpawnLocation, SpawnRotation, ActorSpawnParams);
+		}
+	}
+	
+	// Try and play the sound if specified
+	if (FireSound != nullptr)
+	{
+		UGameplayStatics::PlaySoundAtLocation(this, FireSound, Character->GetActorLocation());
+	}
+	
+	// Try and play a firing animation if specified
+	if (FireAnimation != nullptr)
+	{
+		// Get the animation object for the arms mesh
+		UAnimInstance* AnimInstance = Character->GetMesh1P()->GetAnimInstance();
+		if (AnimInstance != nullptr)
+		{
+			AnimInstance->Montage_Play(FireAnimation, 1.f);
+		}
+	}
+}
+
+bool UTP_FirstPersonWeaponComponent::AttachWeapon(ATP_FirstPersonCharacter* TargetCharacter)
+{
+	Character = TargetCharacter;
+
+	// Check that the character is valid, and has no weapon component yet
+	if (Character == nullptr || Character->GetInstanceComponents().FindItemByClass<UTP_FirstPersonWeaponComponent>())
+	{
+		return false;
+	}
+
+	// Attach the weapon to the First Person Character
+	FAttachmentTransformRules AttachmentRules(EAttachmentRule::SnapToTarget, true);
+	AttachToComponent(Character->GetMesh1P(), AttachmentRules, FName(TEXT("GripPoint")));
+
+	// Set up action bindings
+	if (APlayerController* PlayerController = Cast<APlayerController>(Character->GetController()))
+	{
+		if (UEnhancedInputLocalPlayerSubsystem* Subsystem = ULocalPlayer::GetSubsystem<UEnhancedInputLocalPlayerSubsystem>(PlayerController->GetLocalPlayer()))
+		{
+			// Set the priority of the mapping to 1, so that it overrides the Jump action with the Fire action when using touch input
+			Subsystem->AddMappingContext(FireMappingContext, 1);
+		}
+
+		if (UEnhancedInputComponent* EnhancedInputComponent = Cast<UEnhancedInputComponent>(PlayerController->InputComponent))
+		{
+			// Fire
+			EnhancedInputComponent->BindAction(FireAction, ETriggerEvent::Triggered, this, &UTP_FirstPersonWeaponComponent::Fire);
+		}
+	}
+
+	return true;
+}
+
+void UTP_FirstPersonWeaponComponent::EndPlay(const EEndPlayReason::Type EndPlayReason)
+{
+	// ensure we have a character owner
+	if (Character != nullptr)
+	{
+		// remove the input mapping context from the Player Controller
+		if (APlayerController* PlayerController = Cast<APlayerController>(Character->GetController()))
+		{
+			if (UEnhancedInputLocalPlayerSubsystem* Subsystem = ULocalPlayer::GetSubsystem<UEnhancedInputLocalPlayerSubsystem>(PlayerController->GetLocalPlayer()))
+			{
+				Subsystem->RemoveMappingContext(FireMappingContext);
+			}
+		}
+	}
+
+	// maintain the EndPlay call chain
+	Super::EndPlay(EndPlayReason);
+}
\ No newline at end of file
diff --git a/Templates/TP_FirstPerson/Source/TP_FirstPerson/TP_FirstPersonWeaponComponent.h b/Templates/TP_FirstPerson/Source/TP_FirstPerson/TP_FirstPersonWeaponComponent.h
new file mode 100644
--- /dev/null
+++ b/Templates/TP_FirstPerson/Source/TP_FirstPerson/TP_FirstPersonWeaponComponent.h
@@ -0,0 +1,60 @@
+// Copyright Epic Games, Inc. All Rights Reserved.
+
+#pragma once
+
+#include "CoreMinimal.h"
+#include "Components/SkeletalMeshComponent.h"
+#include "TP_FirstPersonWeaponComponent.generated.h"
+
+class ATP_FirstPersonCharacter;
+
+UCLASS(Blueprintable, BlueprintType, ClassGroup=(Custom), meta=(BlueprintSpawnableComponent) )
+class TP_FIRSTPERSON_API UTP_FirstPersonWeaponComponent : public USkeletalMeshComponent
+{
+	GENERATED_BODY()
+
+public:
+	/** Projectile class to spawn */
+	UPROPERTY(EditDefaultsOnly, Category=Projectile)
+	TSubclassOf<class ATP_FirstPersonProjectile> ProjectileClass;
+
+	/** Sound to play each time we fire */
+	UPROPERTY(EditAnywhere, BlueprintReadWrite, Category=Gameplay)
+	USoundBase* FireSound;
+	
+	/** AnimMontage to play each time we fire */
+	UPROPERTY(EditAnywhere, BlueprintReadWrite, Category = Gameplay)
+	UAnimMontage* FireAnimation;
+
+	/** Gun muzzle's offset from the characters location */
+	UPROPERTY(EditAnywhere, BlueprintReadWrite, Category=Gameplay)
+	FVector MuzzleOffset;
+
+	/** MappingContext */
+	UPROPERTY(EditAnywhere, BlueprintReadOnly, Category=Input, meta=(AllowPrivateAccess = "true"))
+	class UInputMappingContext* FireMappingContext;
+
+	/** Fire Input Action */
+	UPROPERTY(EditAnywhere, BlueprintReadOnly, Category=Input, meta=(AllowPrivateAccess = "true"))
+	class UInputAction* FireAction;
+
+	/** Sets default values for this component's properties */
+	UTP_FirstPersonWeaponComponent();
+
+	/** Attaches the actor to a FirstPersonCharacter */
+	UFUNCTION(BlueprintCallable, Category="Weapon")
+	bool AttachWeapon(ATP_FirstPersonCharacter* TargetCharacter);
+
+	/** Make the weapon Fire a Projectile */
+	UFUNCTION(BlueprintCallable, Category="Weapon")
+	void Fire();
+
+protected:
+	/** Ends gameplay for this component. */
+	UFUNCTION()
+	virtual void EndPlay(const EEndPlayReason::Type EndPlayReason) override;
+
+private:
+	/** The Character holding this weapon*/
+	ATP_FirstPersonCharacter* Character;
+};
diff --git a/Templates/TP_FirstPerson/Source/TP_FirstPerson/TP_PickUpComponent.cpp b/Templates/TP_FirstPerson/Source/TP_FirstPerson/TP_PickUpComponent.cpp
deleted file mode 100644
diff --git a/Templates/TP_FirstPerson/Source/TP_FirstPerson/TP_PickUpComponent.h b/Templates/TP_FirstPerson/Source/TP_FirstPerson/TP_PickUpComponent.h
deleted file mode 100644
diff --git a/Templates/TP_FirstPerson/Source/TP_FirstPerson/TP_WeaponComponent.cpp b/Templates/TP_FirstPerson/Source/TP_FirstPerson/TP_WeaponComponent.cpp
deleted file mode 100644
diff --git a/Templates/TP_FirstPerson/Source/TP_FirstPerson/TP_WeaponComponent.h b/Templates/TP_FirstPerson/Source/TP_FirstPerson/TP_WeaponComponent.h
deleted file mode 100644
diff --git a/Templates/TP_ThirdPerson/Config/config.ini b/Templates/TP_ThirdPerson/Config/config.ini
--- a/Templates/TP_ThirdPerson/Config/config.ini
+++ b/Templates/TP_ThirdPerson/Config/config.ini
@@ -17,8 +17,5 @@ GameName=TP_ThirdPerson
 
 [AdditionalFilesToAdd]
 +Files=Templates/TP_ThirdPerson/Content/*.*
-+Files=Templates/TP_ThirdPerson/Source/TP_ThirdPerson/TP_ThirdPerson.h
-+Files=Templates/TP_ThirdPerson/Source/TP_ThirdPerson/TP_ThirdPersonCharacter.cpp
-+Files=Templates/TP_ThirdPerson/Source/TP_ThirdPerson/TP_ThirdPersonCharacter.h
-+Files=Templates/TP_ThirdPerson/Source/TP_ThirdPerson/TP_ThirdPersonGameMode.cpp
-+Files=Templates/TP_ThirdPerson/Source/TP_ThirdPerson/TP_ThirdPersonGameMode.h
++Files=Templates/TP_ThirdPerson/Source/TP_ThirdPerson/*.h
++Files=Templates/TP_ThirdPerson/Source/TP_ThirdPerson/*.cpp
\ No newline at end of file
diff --git a/Templates/TP_VehicleAdv/Config/config.ini b/Templates/TP_VehicleAdv/Config/config.ini
--- a/Templates/TP_VehicleAdv/Config/config.ini
+++ b/Templates/TP_VehicleAdv/Config/config.ini
@@ -27,17 +27,8 @@ GameName=TP_VehicleAdv
 
 [AdditionalFilesToAdd]
 +Files=Templates/TP_VehicleAdv/Content/*.*
-+Files=Templates/TP_VehicleAdv/Source/TP_VehicleAdv/TP_VehicleAdv.h
-+Files=Templates/TP_VehicleAdv/Source/TP_VehicleAdv/TP_VehicleAdvGameMode.cpp
-+Files=Templates/TP_VehicleAdv/Source/TP_VehicleAdv/TP_VehicleAdvGameMode.h
-+Files=Templates/TP_VehicleAdv/Source/TP_VehicleAdv/TP_VehicleAdvHud.cpp
-+Files=Templates/TP_VehicleAdv/Source/TP_VehicleAdv/TP_VehicleAdvHud.h
-+Files=Templates/TP_VehicleAdv/Source/TP_VehicleAdv/TP_VehicleAdvPawn.cpp
-+Files=Templates/TP_VehicleAdv/Source/TP_VehicleAdv/TP_VehicleAdvPawn.h
-+Files=Templates/TP_VehicleAdv/Source/TP_VehicleAdv/TP_VehicleAdvWheelFront.cpp
-+Files=Templates/TP_VehicleAdv/Source/TP_VehicleAdv/TP_VehicleAdvWheelFront.h
-+Files=Templates/TP_VehicleAdv/Source/TP_VehicleAdv/TP_VehicleAdvWheelRear.cpp
-+Files=Templates/TP_VehicleAdv/Source/TP_VehicleAdv/TP_VehicleAdvWheelRear.h
++Files=Templates/TP_VehicleAdv/Source/TP_VehicleAdv/*.h
++Files=Templates/TP_VehicleAdv/Source/TP_VehicleAdv/*.cpp
 
 [FeaturePackSettings]
 CompileSource=false
