Templates - TP_FirstPerson
Updated template code
Implemented Horror variant
Implemented Shooter variant
#rnx

diff --git a/Templates/TP_FirstPerson/Config/DefaultEngine.ini b/Templates/TP_FirstPerson/Config/DefaultEngine.ini
--- a/Templates/TP_FirstPerson/Config/DefaultEngine.ini
+++ b/Templates/TP_FirstPerson/Config/DefaultEngine.ini
@@ -4,14 +4,14 @@
 +EditProfiles=(Name="Trigger",CustomResponses=((Channel=Projectile, Response=ECR_Ignore)))
 
 [/Script/EngineSettings.GameMapsSettings]
-EditorStartupMap=/Game/FirstPerson/Maps/FirstPersonMap.FirstPersonMap
+EditorStartupMap=/Game/FirstPerson/Lvl_FirstPerson.Lvl_FirstPerson
 LocalMapOptions=
 TransitionMap=
 bUseSplitscreen=True
 TwoPlayerSplitscreenLayout=Horizontal
 ThreePlayerSplitscreenLayout=FavorTop
 GameInstanceClass=/Script/Engine.GameInstance
-GameDefaultMap=/Game/FirstPerson/Maps/FirstPersonMap.FirstPersonMap
+GameDefaultMap=/Game/FirstPerson/Lvl_FirstPerson.Lvl_FirstPerson
 ServerDefaultMap=/Engine/Maps/Entry
 GlobalDefaultGameMode=/Game/FirstPerson/Blueprints/BP_FirstPersonGameMode.BP_FirstPersonGameMode_C
 GlobalDefaultServerGameMode=None
@@ -71,3 +71,6 @@ ConnectionType=USBOnly
 bUseManualIPAddress=False
 ManualIPAddress=
 
+[/Script/AIModule.AISystem]
+bForgetStaleActors=True
+
diff --git a/Templates/TP_FirstPerson/Config/TemplateDefs.ini b/Templates/TP_FirstPerson/Config/TemplateDefs.ini
--- a/Templates/TP_FirstPerson/Config/TemplateDefs.ini
+++ b/Templates/TP_FirstPerson/Config/TemplateDefs.ini
@@ -49,6 +49,8 @@ FilenameReplacements=(Extensions=("cpp","h","ini","cs"),From="%TEMPLATENAME%",To
 ReplacementsInFiles=(Extensions=("cpp","h","ini","cs"),From="%TEMPLATENAME_UPPERCASE%",To="%PROJECTNAME_UPPERCASE%",bCaseSensitive=true)
 ReplacementsInFiles=(Extensions=("cpp","h","ini","cs"),From="%TEMPLATENAME_LOWERCASE%",To="%PROJECTNAME_LOWERCASE%",bCaseSensitive=true)
 ReplacementsInFiles=(Extensions=("cpp","h","ini","cs"),From="%TEMPLATENAME%",To="%PROJECTNAME%",bCaseSensitive=false)
-SharedContentPacks=(MountName="FirstPersonArms",DetailLevels=("Standard"))
-SharedContentPacks=(MountName="FPWeapon",DetailLevels=("Standard"))
+SharedContentPacks=(MountName="Characters",DetailLevels=("High"))
 SharedContentPacks=(MountName="LevelPrototyping",DetailLevels=("High"))
+SharedContentPacks=(MountName="Input",DetailLevels=("High"))
+SharedContentPacks=(MountName="Weapons",DetailLevels=("Standard"))
+EditDetailLevelPreference="High"
\ No newline at end of file
diff --git a/Templates/TP_FirstPerson/Source/TP_FirstPerson/TP_FirstPerson.Build.cs b/Templates/TP_FirstPerson/Source/TP_FirstPerson/TP_FirstPerson.Build.cs
--- a/Templates/TP_FirstPerson/Source/TP_FirstPerson/TP_FirstPerson.Build.cs
+++ b/Templates/TP_FirstPerson/Source/TP_FirstPerson/TP_FirstPerson.Build.cs
@@ -8,6 +8,33 @@ public class TP_FirstPerson : ModuleRules
 	{
 		PCHUsage = PCHUsageMode.UseExplicitOrSharedPCHs;
 
-		PublicDependencyModuleNames.AddRange(new string[] { "Core", "CoreUObject", "Engine", "InputCore", "EnhancedInput" });
+		PublicDependencyModuleNames.AddRange(new string[] {
+			"Core",
+			"CoreUObject",
+			"Engine",
+			"InputCore",
+			"EnhancedInput",
+			"AIModule",
+			"StateTreeModule",
+			"GameplayStateTreeModule",
+			"UMG"
+		});
+
+		PrivateDependencyModuleNames.AddRange(new string[] { });
+
+		PublicIncludePaths.AddRange(new string[] {
+			"TP_FirstPerson",
+			"TP_FirstPerson/Variant_Horror",
+			"TP_FirstPerson/Variant_Shooter",
+			"TP_FirstPerson/Variant_Shooter/AI"
+		});
+
+		// Uncomment if you are using Slate UI
+		// PrivateDependencyModuleNames.AddRange(new string[] { "Slate", "SlateCore" });
+
+		// Uncomment if you are using online features
+		// PrivateDependencyModuleNames.Add("OnlineSubsystem");
+
+		// To include OnlineSubsystemSteam, add it to the plugins section in your uproject file with the Enabled attribute set to true
 	}
 }
diff --git a/Templates/TP_FirstPerson/Source/TP_FirstPerson/TP_FirstPersonCameraManager.cpp b/Templates/TP_FirstPerson/Source/TP_FirstPerson/TP_FirstPersonCameraManager.cpp
new file mode 100644
--- /dev/null
+++ b/Templates/TP_FirstPerson/Source/TP_FirstPerson/TP_FirstPersonCameraManager.cpp
@@ -0,0 +1,11 @@
+// Copyright Epic Games, Inc. All Rights Reserved.
+
+
+#include "TP_FirstPersonCameraManager.h"
+
+ATP_FirstPersonCameraManager::ATP_FirstPersonCameraManager()
+{
+	// set the min/max pitch
+	ViewPitchMin = -70.0f;
+	ViewPitchMax = 80.0f;
+}
diff --git a/Templates/TP_FirstPerson/Source/TP_FirstPerson/TP_FirstPersonCameraManager.h b/Templates/TP_FirstPerson/Source/TP_FirstPerson/TP_FirstPersonCameraManager.h
new file mode 100644
--- /dev/null
+++ b/Templates/TP_FirstPerson/Source/TP_FirstPerson/TP_FirstPersonCameraManager.h
@@ -0,0 +1,22 @@
+// Copyright Epic Games, Inc. All Rights Reserved.
+
+#pragma once
+
+#include "CoreMinimal.h"
+#include "Camera/PlayerCameraManager.h"
+#include "TP_FirstPersonCameraManager.generated.h"
+
+/**
+ *  Basic First Person camera manager.
+ *  Limits min/max look pitch.
+ */
+UCLASS()
+class ATP_FirstPersonCameraManager : public APlayerCameraManager
+{
+	GENERATED_BODY()
+	
+public:
+
+	/** Constructor */
+	ATP_FirstPersonCameraManager();
+};
diff --git a/Templates/TP_FirstPerson/Source/TP_FirstPerson/TP_FirstPersonCharacter.cpp b/Templates/TP_FirstPerson/Source/TP_FirstPerson/TP_FirstPersonCharacter.cpp
--- a/Templates/TP_FirstPerson/Source/TP_FirstPerson/TP_FirstPersonCharacter.cpp
+++ b/Templates/TP_FirstPerson/Source/TP_FirstPerson/TP_FirstPersonCharacter.cpp
@@ -1,15 +1,13 @@
 // Copyright Epic Games, Inc. All Rights Reserved.
 
 #include "TP_FirstPersonCharacter.h"
-#include "TP_FirstPersonProjectile.h"
 #include "Animation/AnimInstance.h"
 #include "Camera/CameraComponent.h"
 #include "Components/CapsuleComponent.h"
 #include "Components/SkeletalMeshComponent.h"
 #include "EnhancedInputComponent.h"
-#include "EnhancedInputSubsystems.h"
 #include "InputActionValue.h"
-#include "Engine/LocalPlayer.h"
+#include "GameFramework/CharacterMovementComponent.h"
 
 DEFINE_LOG_CATEGORY(LogTemplateCharacter);
 
@@ -20,37 +18,34 @@ ATP_FirstPersonCharacter::ATP_FirstPersonCharacter()
 {
 	// Set size for collision capsule
 	GetCapsuleComponent()->InitCapsuleSize(55.f, 96.0f);
-		
-	// Create a CameraComponent	
-	FirstPersonCameraComponent = CreateDefaultSubobject<UCameraComponent>(TEXT("FirstPersonCamera"));
-	FirstPersonCameraComponent->SetupAttachment(GetCapsuleComponent());
-	FirstPersonCameraComponent->SetRelativeLocation(FVector(-10.f, 0.f, 60.f)); // Position the camera
+	
+	// Create the first person mesh that will be viewed only by this character's owner
+	FirstPersonMesh = CreateDefaultSubobject<USkeletalMeshComponent>(TEXT("First Person Mesh"));
+
+	FirstPersonMesh->SetupAttachment(GetMesh());
+	FirstPersonMesh->SetOnlyOwnerSee(true);
+	FirstPersonMesh->FirstPersonPrimitiveType = EFirstPersonPrimitiveType::FirstPerson;
+	FirstPersonMesh->SetCollisionProfileName(FName("NoCollision"));
+
+	// Create the Camera Component	
+	FirstPersonCameraComponent = CreateDefaultSubobject<UCameraComponent>(TEXT("First Person Camera"));
+	FirstPersonCameraComponent->SetupAttachment(FirstPersonMesh, FName("head"));
+	FirstPersonCameraComponent->SetRelativeLocationAndRotation(FVector(-2.8f, 5.89f, 0.0f), FRotator(0.0f, 90.0f, -90.0f));
 	FirstPersonCameraComponent->bUsePawnControlRotation = true;
+	FirstPersonCameraComponent->bEnableFirstPersonFieldOfView = true;
+	FirstPersonCameraComponent->bEnableFirstPersonScale = true;
+	FirstPersonCameraComponent->FirstPersonFieldOfView = 70.0f;
+	FirstPersonCameraComponent->FirstPersonScale = 0.6f;
 
-	// Create a mesh component that will be used when being viewed from a '1st person' view (when controlling this pawn)
-	Mesh1P = CreateDefaultSubobject<USkeletalMeshComponent>(TEXT("CharacterMesh1P"));
-	Mesh1P->SetOnlyOwnerSee(true);
-	Mesh1P->SetupAttachment(FirstPersonCameraComponent);
-	Mesh1P->bCastDynamicShadow = false;
-	Mesh1P->CastShadow = false;
-	Mesh1P->SetRelativeLocation(FVector(-30.f, 0.f, -150.f));
+	// configure the character comps
+	GetMesh()->SetOwnerNoSee(true);
+	GetMesh()->FirstPersonPrimitiveType = EFirstPersonPrimitiveType::WorldSpaceRepresentation;
 
-}
-
-//////////////////////////////////////////////////////////////////////////// Input
-
-void ATP_FirstPersonCharacter::NotifyControllerChanged()
-{
-	Super::NotifyControllerChanged();
+	GetCapsuleComponent()->SetCapsuleSize(34.0f, 96.0f);
 
-	// Add Input Mapping Context
-	if (APlayerController* PlayerController = Cast<APlayerController>(GetController()))
-	{
-		if (UEnhancedInputLocalPlayerSubsystem* Subsystem = ULocalPlayer::GetSubsystem<UEnhancedInputLocalPlayerSubsystem>(PlayerController->GetLocalPlayer()))
-		{
-			Subsystem->AddMappingContext(DefaultMappingContext, 0);
-		}
-	}
+	// Configure character movement
+	GetCharacterMovement()->BrakingDecelerationFalling = 1500.0f;
+	GetCharacterMovement()->AirControl = 0.5f;
 }
 
 void ATP_FirstPersonCharacter::SetupPlayerInputComponent(UInputComponent* PlayerInputComponent)
@@ -59,14 +54,15 @@ void ATP_FirstPersonCharacter::SetupPlayerInputComponent(UInputComponent* Player
 	if (UEnhancedInputComponent* EnhancedInputComponent = Cast<UEnhancedInputComponent>(PlayerInputComponent))
 	{
 		// Jumping
-		EnhancedInputComponent->BindAction(JumpAction, ETriggerEvent::Started, this, &ACharacter::Jump);
-		EnhancedInputComponent->BindAction(JumpAction, ETriggerEvent::Completed, this, &ACharacter::StopJumping);
+		EnhancedInputComponent->BindAction(JumpAction, ETriggerEvent::Started, this, &ATP_FirstPersonCharacter::DoJumpStart);
+		EnhancedInputComponent->BindAction(JumpAction, ETriggerEvent::Completed, this, &ATP_FirstPersonCharacter::DoJumpEnd);
 
 		// Moving
-		EnhancedInputComponent->BindAction(MoveAction, ETriggerEvent::Triggered, this, &ATP_FirstPersonCharacter::Move);
+		EnhancedInputComponent->BindAction(MoveAction, ETriggerEvent::Triggered, this, &ATP_FirstPersonCharacter::MoveInput);
 
-		// Looking
-		EnhancedInputComponent->BindAction(LookAction, ETriggerEvent::Triggered, this, &ATP_FirstPersonCharacter::Look);
+		// Looking/Aiming
+		EnhancedInputComponent->BindAction(LookAction, ETriggerEvent::Triggered, this, &ATP_FirstPersonCharacter::LookInput);
+		EnhancedInputComponent->BindAction(MouseLookAction, ETriggerEvent::Triggered, this, &ATP_FirstPersonCharacter::LookInput);
 	}
 	else
 	{
@@ -75,28 +71,54 @@ void ATP_FirstPersonCharacter::SetupPlayerInputComponent(UInputComponent* Player
 }
 
 
-void ATP_FirstPersonCharacter::Move(const FInputActionValue& Value)
+void ATP_FirstPersonCharacter::MoveInput(const FInputActionValue& Value)
 {
-	// input is a Vector2D
+	// get the Vector2D move axis
 	FVector2D MovementVector = Value.Get<FVector2D>();
 
-	if (GetController() != nullptr)
-	{
-		// add movement 
-		AddMovementInput(GetActorForwardVector(), MovementVector.Y);
-		AddMovementInput(GetActorRightVector(), MovementVector.X);
-	}
+	// pass the axis values to the move input
+	DoMove(MovementVector.X, MovementVector.Y);
+
 }
 
-void ATP_FirstPersonCharacter::Look(const FInputActionValue& Value)
+void ATP_FirstPersonCharacter::LookInput(const FInputActionValue& Value)
 {
-	// input is a Vector2D
+	// get the Vector2D look axis
 	FVector2D LookAxisVector = Value.Get<FVector2D>();
 
-	if (GetController() != nullptr)
+	// pass the axis values to the aim input
+	DoAim(LookAxisVector.X, LookAxisVector.Y);
+
+}
+
+void ATP_FirstPersonCharacter::DoAim(float Yaw, float Pitch)
+{
+	if (GetController())
 	{
-		// add yaw and pitch input to controller
-		AddControllerYawInput(LookAxisVector.X);
-		AddControllerPitchInput(LookAxisVector.Y);
+		// pass the rotation inputs
+		AddControllerYawInput(Yaw);
+		AddControllerPitchInput(Pitch);
 	}
-}
\ No newline at end of file
+}
+
+void ATP_FirstPersonCharacter::DoMove(float Right, float Forward)
+{
+	if (GetController())
+	{
+		// pass the move inputs
+		AddMovementInput(GetActorRightVector(), Right);
+		AddMovementInput(GetActorForwardVector(), Forward);
+	}
+}
+
+void ATP_FirstPersonCharacter::DoJumpStart()
+{
+	// pass Jump to the character
+	Jump();
+}
+
+void ATP_FirstPersonCharacter::DoJumpEnd()
+{
+	// pass StopJumping to the character
+	StopJumping();
+}
diff --git a/Templates/TP_FirstPerson/Source/TP_FirstPerson/TP_FirstPersonCharacter.h b/Templates/TP_FirstPerson/Source/TP_FirstPerson/TP_FirstPersonCharacter.h
--- a/Templates/TP_FirstPerson/Source/TP_FirstPerson/TP_FirstPersonCharacter.h
+++ b/Templates/TP_FirstPerson/Source/TP_FirstPerson/TP_FirstPersonCharacter.h
@@ -11,60 +11,83 @@ class UInputComponent;
 class USkeletalMeshComponent;
 class UCameraComponent;
 class UInputAction;
-class UInputMappingContext;
 struct FInputActionValue;
 
 DECLARE_LOG_CATEGORY_EXTERN(LogTemplateCharacter, Log, All);
 
-UCLASS(config=Game)
+/**
+ *  A basic first person character
+ */
+UCLASS(abstract)
 class ATP_FirstPersonCharacter : public ACharacter
 {
 	GENERATED_BODY()
 
-	/** Pawn mesh: 1st person view (arms; seen only by self) */
-	UPROPERTY(VisibleAnywhere, BlueprintReadOnly, Category=Mesh, meta = (AllowPrivateAccess = "true"))
-	USkeletalMeshComponent* Mesh1P;
+	/** Pawn mesh: first person view (arms; seen only by self) */
+	UPROPERTY(VisibleAnywhere, BlueprintReadOnly, Category="Components", meta = (AllowPrivateAccess = "true"))
+	USkeletalMeshComponent* FirstPersonMesh;
 
 	/** First person camera */
-	UPROPERTY(VisibleAnywhere, BlueprintReadOnly, Category = Camera, meta = (AllowPrivateAccess = "true"))
+	UPROPERTY(VisibleAnywhere, BlueprintReadOnly, Category="Components", meta = (AllowPrivateAccess = "true"))
 	UCameraComponent* FirstPersonCameraComponent;
 
-	/** MappingContext */
-	UPROPERTY(EditAnywhere, BlueprintReadOnly, Category = Input, meta = (AllowPrivateAccess = "true"))
-	UInputMappingContext* DefaultMappingContext;
+protected:
 
 	/** Jump Input Action */
-	UPROPERTY(EditAnywhere, BlueprintReadOnly, Category=Input, meta=(AllowPrivateAccess = "true"))
+	UPROPERTY(EditAnywhere, BlueprintReadOnly, Category ="Input")
 	UInputAction* JumpAction;
 
 	/** Move Input Action */
-	UPROPERTY(EditAnywhere, BlueprintReadOnly, Category=Input, meta=(AllowPrivateAccess = "true"))
+	UPROPERTY(EditAnywhere, BlueprintReadOnly, Category ="Input")
 	UInputAction* MoveAction;
 
 	/** Look Input Action */
-	UPROPERTY(EditAnywhere, BlueprintReadOnly, Category = Input, meta = (AllowPrivateAccess = "true"))
+	UPROPERTY(EditAnywhere, BlueprintReadOnly, Category ="Input")
 	class UInputAction* LookAction;
+
+	/** Mouse Look Input Action */
+	UPROPERTY(EditAnywhere, BlueprintReadOnly, Category ="Input")
+	class UInputAction* MouseLookAction;
 	
 public:
 	ATP_FirstPersonCharacter();
 
 protected:
-	/** Called for movement input */
-	void Move(const FInputActionValue& Value);
 
-	/** Called for looking input */
-	void Look(const FInputActionValue& Value);
+	/** Called from Input Actions for movement input */
+	void MoveInput(const FInputActionValue& Value);
+
+	/** Called from Input Actions for looking input */
+	void LookInput(const FInputActionValue& Value);
+
+	/** Handles aim inputs from either controls or UI interfaces */
+	UFUNCTION(BlueprintCallable, Category="Input")
+	virtual void DoAim(float Yaw, float Pitch);
+
+	/** Handles move inputs from either controls or UI interfaces */
+	UFUNCTION(BlueprintCallable, Category="Input")
+	virtual void DoMove(float Right, float Forward);
+
+	/** Handles jump start inputs from either controls or UI interfaces */
+	UFUNCTION(BlueprintCallable, Category="Input")
+	virtual void DoJumpStart();
+
+	/** Handles jump end inputs from either controls or UI interfaces */
+	UFUNCTION(BlueprintCallable, Category="Input")
+	virtual void DoJumpEnd();
 
 protected:
-	// APawn interface
-	virtual void NotifyControllerChanged() override;
+
+	/** Set up input action bindings */
 	virtual void SetupPlayerInputComponent(UInputComponent* InputComponent) override;
-	// End of APawn interface
+	
 
 public:
-	/** Returns Mesh1P subobject **/
-	USkeletalMeshComponent* GetMesh1P() const { return Mesh1P; }
-	/** Returns FirstPersonCameraComponent subobject **/
+
+	/** Returns the first person mesh **/
+	USkeletalMeshComponent* GetFirstPersonMesh() const { return FirstPersonMesh; }
+
+	/** Returns first person camera component **/
 	UCameraComponent* GetFirstPersonCameraComponent() const { return FirstPersonCameraComponent; }
 
 };
diff --git a/Templates/TP_FirstPerson/Source/TP_FirstPerson/TP_FirstPersonGameMode.cpp b/Templates/TP_FirstPerson/Source/TP_FirstPerson/TP_FirstPersonGameMode.cpp
--- a/Templates/TP_FirstPerson/Source/TP_FirstPerson/TP_FirstPersonGameMode.cpp
+++ b/Templates/TP_FirstPerson/Source/TP_FirstPerson/TP_FirstPersonGameMode.cpp
@@ -1,14 +1,8 @@
 // Copyright Epic Games, Inc. All Rights Reserved.
 
 #include "TP_FirstPersonGameMode.h"
-#include "TP_FirstPersonCharacter.h"
-#include "UObject/ConstructorHelpers.h"
 
 ATP_FirstPersonGameMode::ATP_FirstPersonGameMode()
-	: Super()
 {
-	// set default pawn class to our Blueprinted character
-	static ConstructorHelpers::FClassFinder<APawn> PlayerPawnClassFinder(TEXT("/Game/FirstPerson/Blueprints/BP_FirstPersonCharacter"));
-	DefaultPawnClass = PlayerPawnClassFinder.Class;
-
+	// stub
 }
diff --git a/Templates/TP_FirstPerson/Source/TP_FirstPerson/TP_FirstPersonGameMode.h b/Templates/TP_FirstPerson/Source/TP_FirstPerson/TP_FirstPersonGameMode.h
--- a/Templates/TP_FirstPerson/Source/TP_FirstPerson/TP_FirstPersonGameMode.h
+++ b/Templates/TP_FirstPerson/Source/TP_FirstPerson/TP_FirstPersonGameMode.h
@@ -6,7 +6,10 @@
 #include "GameFramework/GameModeBase.h"
 #include "TP_FirstPersonGameMode.generated.h"
 
-UCLASS(minimalapi)
+/**
+ *  Simple GameMode for a first person game
+ */
+UCLASS(abstract)
 class ATP_FirstPersonGameMode : public AGameModeBase
 {
 	GENERATED_BODY()
diff --git a/Templates/TP_FirstPerson/Source/TP_FirstPerson/TP_FirstPersonPickUpComponent.cpp b/Templates/TP_FirstPerson/Source/TP_FirstPerson/TP_FirstPersonPickUpComponent.cpp
deleted file mode 100644
diff --git a/Templates/TP_FirstPerson/Source/TP_FirstPerson/TP_FirstPersonPickUpComponent.h b/Templates/TP_FirstPerson/Source/TP_FirstPerson/TP_FirstPersonPickUpComponent.h
deleted file mode 100644
diff --git a/Templates/TP_FirstPerson/Source/TP_FirstPerson/TP_FirstPersonPlayerController.cpp b/Templates/TP_FirstPerson/Source/TP_FirstPerson/TP_FirstPersonPlayerController.cpp
new file mode 100644
--- /dev/null
+++ b/Templates/TP_FirstPerson/Source/TP_FirstPerson/TP_FirstPersonPlayerController.cpp
@@ -0,0 +1,28 @@
+// Copyright Epic Games, Inc. All Rights Reserved.
+
+
+#include "TP_FirstPersonPlayerController.h"
+#include "EnhancedInputSubsystems.h"
+#include "Engine/LocalPlayer.h"
+#include "InputMappingContext.h"
+#include "TP_FirstPersonCameraManager.h"
+
+ATP_FirstPersonPlayerController::ATP_FirstPersonPlayerController()
+{
+	// set the player camera manager class
+	PlayerCameraManagerClass = ATP_FirstPersonCameraManager::StaticClass();
+}
+
+void ATP_FirstPersonPlayerController::SetupInputComponent()
+{
+	Super::SetupInputComponent();
+
+	// Add Input Mapping Context
+	if (UEnhancedInputLocalPlayerSubsystem* Subsystem = ULocalPlayer::GetSubsystem<UEnhancedInputLocalPlayerSubsystem>(GetLocalPlayer()))
+	{
+		for (UInputMappingContext* CurrentContext : DefaultMappingContexts)
+		{
+			Subsystem->AddMappingContext(CurrentContext, 0);
+		}
+	}
+}
diff --git a/Templates/TP_FirstPerson/Source/TP_FirstPerson/TP_FirstPersonPlayerController.h b/Templates/TP_FirstPerson/Source/TP_FirstPerson/TP_FirstPersonPlayerController.h
new file mode 100644
--- /dev/null
+++ b/Templates/TP_FirstPerson/Source/TP_FirstPerson/TP_FirstPersonPlayerController.h
@@ -0,0 +1,35 @@
+// Copyright Epic Games, Inc. All Rights Reserved.
+
+#pragma once
+
+#include "CoreMinimal.h"
+#include "GameFramework/PlayerController.h"
+#include "TP_FirstPersonPlayerController.generated.h"
+
+class UInputMappingContext;
+
+/**
+ *  Simple first person Player Controller
+ *  Manages the input mapping context.
+ *  Overrides the Player Camera Manager class.
+ */
+UCLASS(abstract)
+class TP_FIRSTPERSON_API ATP_FirstPersonPlayerController : public APlayerController
+{
+	GENERATED_BODY()
+	
+public:
+
+	/** Constructor */
+	ATP_FirstPersonPlayerController();
+
+protected:
+
+	/** Input Mapping Contexts */
+	UPROPERTY(EditAnywhere, BlueprintReadOnly, Category ="Input", meta = (AllowPrivateAccess = "true"))
+	TArray<UInputMappingContext*> DefaultMappingContexts;
+
+	/** Input mapping context setup */
+	virtual void SetupInputComponent() override;
+
+};
diff --git a/Templates/TP_FirstPerson/Source/TP_FirstPerson/TP_FirstPersonProjectile.cpp b/Templates/TP_FirstPerson/Source/TP_FirstPerson/TP_FirstPersonProjectile.cpp
deleted file mode 100644
diff --git a/Templates/TP_FirstPerson/Source/TP_FirstPerson/TP_FirstPersonProjectile.h b/Templates/TP_FirstPerson/Source/TP_FirstPerson/TP_FirstPersonProjectile.h
deleted file mode 100644
diff --git a/Templates/TP_FirstPerson/Source/TP_FirstPerson/TP_FirstPersonWeaponComponent.cpp b/Templates/TP_FirstPerson/Source/TP_FirstPerson/TP_FirstPersonWeaponComponent.cpp
deleted file mode 100644
diff --git a/Templates/TP_FirstPerson/Source/TP_FirstPerson/TP_FirstPersonWeaponComponent.h b/Templates/TP_FirstPerson/Source/TP_FirstPerson/TP_FirstPersonWeaponComponent.h
deleted file mode 100644
diff --git a/Templates/TP_FirstPerson/Source/TP_FirstPerson/Variant_Horror/HorrorCharacter.cpp b/Templates/TP_FirstPerson/Source/TP_FirstPerson/Variant_Horror/HorrorCharacter.cpp
new file mode 100644
--- /dev/null
+++ b/Templates/TP_FirstPerson/Source/TP_FirstPerson/Variant_Horror/HorrorCharacter.cpp
@@ -0,0 +1,143 @@
+// Copyright Epic Games, Inc. All Rights Reserved.
+
+
+#include "Variant_Horror/HorrorCharacter.h"
+#include "Engine/World.h"
+#include "TimerManager.h"
+#include "GameFramework/CharacterMovementComponent.h"
+#include "Camera/CameraComponent.h"
+#include "Components/SpotLightComponent.h"
+#include "EnhancedInputComponent.h"
+#include "InputAction.h"
+
+AHorrorCharacter::AHorrorCharacter()
+{
+	// create the spotlight
+	SpotLight = CreateDefaultSubobject<USpotLightComponent>(TEXT("SpotLight"));
+	SpotLight->SetupAttachment(GetFirstPersonCameraComponent());
+
+	SpotLight->SetRelativeLocationAndRotation(FVector(30.0f, 17.5f, -5.0f), FRotator(-18.6f, -1.3f, 5.26f));
+	SpotLight->Intensity = 0.5;
+	SpotLight->SetIntensityUnits(ELightUnits::Lumens);
+	SpotLight->AttenuationRadius = 1050.0f;
+	SpotLight->InnerConeAngle = 18.7f;
+	SpotLight->OuterConeAngle = 45.24f;
+}
+
+void AHorrorCharacter::BeginPlay()
+{
+	Super::BeginPlay();
+
+	// initialize sprint meter to max
+	SprintMeter = SprintTime;
+
+	// Initialize the walk speed
+	GetCharacterMovement()->MaxWalkSpeed = WalkSpeed;
+
+	// start the sprint tick timer
+	GetWorld()->GetTimerManager().SetTimer(SprintTimer, this, &AHorrorCharacter::SprintFixedTick, SprintFixedTickTime, true);
+}
+
+void AHorrorCharacter::EndPlay(EEndPlayReason::Type EndPlayReason)
+{
+	Super::EndPlay(EndPlayReason);
+
+	// clear the sprint timer
+	GetWorld()->GetTimerManager().ClearTimer(SprintTimer);
+}
+
+void AHorrorCharacter::SetupPlayerInputComponent(UInputComponent* PlayerInputComponent)
+{
+	Super::SetupPlayerInputComponent(PlayerInputComponent);
+
+	{
+		// Set up action bindings
+		if (UEnhancedInputComponent* EnhancedInputComponent = Cast<UEnhancedInputComponent>(PlayerInputComponent))
+		{
+			// Sprinting
+			EnhancedInputComponent->BindAction(SprintAction, ETriggerEvent::Started, this, &AHorrorCharacter::DoStartSprint);
+			EnhancedInputComponent->BindAction(SprintAction, ETriggerEvent::Completed, this, &AHorrorCharacter::DoEndSprint);
+
+		}
+	}
+}
+
+void AHorrorCharacter::DoStartSprint()
+{
+	// set the sprinting flag
+	bSprinting = true;
+
+	// are we out of recovery mode?
+	if (!bRecovering)
+	{
+		// set the sprint walk speed
+		GetCharacterMovement()->MaxWalkSpeed = SprintSpeed;
+
+		// call the sprint state changed delegate
+		OnSprintStateChanged.Broadcast(true);
+	}
+
+}
+
+void AHorrorCharacter::DoEndSprint()
+{
+	// set the sprinting flag
+	bSprinting = false;
+
+	// are we out of recovery mode?
+	if (!bRecovering)
+	{
+		// set the default walk speed
+		GetCharacterMovement()->MaxWalkSpeed = WalkSpeed;
+
+		// call the sprint state changed delegate
+		OnSprintStateChanged.Broadcast(false);
+	}
+}
+
+void AHorrorCharacter::SprintFixedTick()
+{
+	// are we out of recovery and still have stamina?
+	if (bSprinting && !bRecovering)
+	{
+
+		// do we still have meter to burn?
+		if (SprintMeter > 0.0f)
+		{
+			// update the sprint meter
+			SprintMeter = FMath::Max(SprintMeter - SprintFixedTickTime, 0.0f);
+
+			// have we run out of stamina?
+			if (SprintMeter <= 0.0f)
+			{
+				// raise the recovering flag
+				bRecovering = true;
+
+				// set the recovering walk speed
+				GetCharacterMovement()->MaxWalkSpeed = RecoveringWalkSpeed;
+			}
+		}
+		
+	} else {
+
+		// recover stamina
+		SprintMeter = FMath::Min(SprintMeter + SprintFixedTickTime, SprintTime);
+
+		if (SprintMeter >= SprintTime)
+		{
+			// lower the recovering flag
+			bRecovering = false;
+
+			// set the walk or sprint speed depending on whether the sprint button is down
+			GetCharacterMovement()->MaxWalkSpeed = bSprinting ? SprintSpeed : WalkSpeed;
+
+			// update the sprint state depending on whether the button is down or not
+			OnSprintStateChanged.Broadcast(bSprinting);
+		}
+
+	}
+
+	// broadcast the sprint meter updated delegate
+	OnSprintMeterUpdated.Broadcast(SprintMeter / SprintTime);
+
+}
diff --git a/Templates/TP_FirstPerson/Source/TP_FirstPerson/Variant_Horror/HorrorCharacter.h b/Templates/TP_FirstPerson/Source/TP_FirstPerson/Variant_Horror/HorrorCharacter.h
new file mode 100644
--- /dev/null
+++ b/Templates/TP_FirstPerson/Source/TP_FirstPerson/Variant_Horror/HorrorCharacter.h
@@ -0,0 +1,104 @@
+// Copyright Epic Games, Inc. All Rights Reserved.
+
+#pragma once
+
+#include "CoreMinimal.h"
+#include "TP_FirstPersonCharacter.h"
+#include "HorrorCharacter.generated.h"
+
+class USpotLightComponent;
+class UInputAction;
+
+DECLARE_DYNAMIC_MULTICAST_DELEGATE_OneParam(FUpdateSprintMeterDelegate, float, Percentage);
+DECLARE_DYNAMIC_MULTICAST_DELEGATE_OneParam(FSprintStateChangedDelegate, bool, bSprinting);
+
+/**
+ *  Simple first person horror character
+ *  Provides stamina-based sprinting
+ */
+UCLASS(abstract)
+class TP_FIRSTPERSON_API AHorrorCharacter : public ATP_FirstPersonCharacter
+{
+	GENERATED_BODY()
+
+	/** Player light source */
+	UPROPERTY(VisibleAnywhere, BlueprintReadOnly, Category="Components", meta = (AllowPrivateAccess = "true"))
+	USpotLightComponent* SpotLight;
+	
+protected:
+
+	/** Fire weapon input action */
+	UPROPERTY(EditAnywhere, BlueprintReadOnly, Category ="Input")
+	UInputAction* SprintAction;
+
+	/** If true, we're sprinting */
+	bool bSprinting = false;
+
+	/** If true, we're recovering stamina */
+	bool bRecovering = false;
+
+	/** Default walk speed when not sprinting or recovering */
+	UPROPERTY(EditAnywhere, Category="Walk")
+	float WalkSpeed = 250.0f;
+
+	/** Time interval for sprinting stamina ticks */
+	UPROPERTY(EditAnywhere, Category="Sprint")
+	float SprintFixedTickTime = 0.03333f;
+
+	/** Sprint stamina amount. Maxes at SprintTime */
+	float SprintMeter = 0.0f;
+
+	/** How long we can sprint for, in seconds */
+	UPROPERTY(EditAnywhere, Category="Sprint", meta=(ClampMin = 0, ClampMax = 10, Units = "s"))
+	float SprintTime = 3.0f;
+
+	/** Walk speed while sprinting */
+	UPROPERTY(EditAnywhere, Category="Sprint", meta=(ClampMin = 0, ClampMax = 10, Units = "cm/s"))
+	float SprintSpeed = 600.0f;
+
+	/** Walk speed while recovering stamina */
+	UPROPERTY(EditAnywhere, Category="Recovery")
+	float RecoveringWalkSpeed = 150.0f;
+
+	/** Time it takes for the sprint meter to recover */
+	UPROPERTY(EditAnywhere, Category="Recovery", meta=(ClampMin = 0, ClampMax = 10, Units = "s"))
+	float RecoveryTime = 0.0f;
+
+	/** Sprint tick timer */
+	FTimerHandle SprintTimer;
+
+public:
+
+	/** Delegate called when the sprint meter should be updated */
+	FUpdateSprintMeterDelegate OnSprintMeterUpdated;
+
+	/** Delegate called when we start and stop sprinting */
+	FSprintStateChangedDelegate OnSprintStateChanged;
+
+protected:
+
+	/** Constructor */
+	AHorrorCharacter();
+
+	/** Gameplay initialization */
+	virtual void BeginPlay() override;
+
+	/** Gameplay cleanup */
+	virtual void EndPlay(EEndPlayReason::Type EndPlayReason) override;
+
+	/** Set up input action bindings */
+	virtual void SetupPlayerInputComponent(UInputComponent* InputComponent) override;
+
+protected:
+
+	/** Starts sprinting behavior */
+	UFUNCTION(BlueprintCallable, Category="Input")
+	void DoStartSprint();
+
+	/** Stops sprinting behavior */
+	UFUNCTION(BlueprintCallable, Category="Input")
+	void DoEndSprint();
+
+	/** Called while sprinting at a fixed time interval */
+	void SprintFixedTick();
+};
diff --git a/Templates/TP_FirstPerson/Source/TP_FirstPerson/Variant_Horror/HorrorGameMode.cpp b/Templates/TP_FirstPerson/Source/TP_FirstPerson/Variant_Horror/HorrorGameMode.cpp
new file mode 100644
--- /dev/null
+++ b/Templates/TP_FirstPerson/Source/TP_FirstPerson/Variant_Horror/HorrorGameMode.cpp
@@ -0,0 +1,9 @@
+// Copyright Epic Games, Inc. All Rights Reserved.
+
+
+#include "Variant_Horror/HorrorGameMode.h"
+
+AHorrorGameMode::AHorrorGameMode()
+{
+	// stub
+}
diff --git a/Templates/TP_FirstPerson/Source/TP_FirstPerson/Variant_Horror/HorrorGameMode.h b/Templates/TP_FirstPerson/Source/TP_FirstPerson/Variant_Horror/HorrorGameMode.h
new file mode 100644
--- /dev/null
+++ b/Templates/TP_FirstPerson/Source/TP_FirstPerson/Variant_Horror/HorrorGameMode.h
@@ -0,0 +1,21 @@
+// Copyright Epic Games, Inc. All Rights Reserved.
+
+#pragma once
+
+#include "CoreMinimal.h"
+#include "GameFramework/GameModeBase.h"
+#include "HorrorGameMode.generated.h"
+
+/**
+ *  Simple GameMode for a first person horror game
+ */
+UCLASS(abstract)
+class TP_FIRSTPERSON_API AHorrorGameMode : public AGameModeBase
+{
+	GENERATED_BODY()
+	
+public:
+
+	/** Constructor */
+	AHorrorGameMode();
+};
diff --git a/Templates/TP_FirstPerson/Source/TP_FirstPerson/Variant_Horror/HorrorPlayerController.cpp b/Templates/TP_FirstPerson/Source/TP_FirstPerson/Variant_Horror/HorrorPlayerController.cpp
new file mode 100644
--- /dev/null
+++ b/Templates/TP_FirstPerson/Source/TP_FirstPerson/Variant_Horror/HorrorPlayerController.cpp
@@ -0,0 +1,48 @@
+// Copyright Epic Games, Inc. All Rights Reserved.
+
+
+#include "Variant_Horror/HorrorPlayerController.h"
+#include "EnhancedInputSubsystems.h"
+#include "Engine/LocalPlayer.h"
+#include "InputMappingContext.h"
+#include "TP_FirstPersonCameraManager.h"
+#include "HorrorCharacter.h"
+#include "HorrorUI.h"
+
+AHorrorPlayerController::AHorrorPlayerController()
+{
+	// set the player camera manager class
+	PlayerCameraManagerClass = ATP_FirstPersonCameraManager::StaticClass();
+}
+
+void AHorrorPlayerController::OnPossess(APawn* aPawn)
+{
+	Super::OnPossess(aPawn);
+
+	// set up the UI for the character
+	if (AHorrorCharacter* HorrorCharacter = Cast<AHorrorCharacter>(aPawn))
+	{
+		// create the UI
+		if (!HorrorUI)
+		{
+			HorrorUI = CreateWidget<UHorrorUI>(this, HorrorUIClass);
+			HorrorUI->AddToViewport(0);
+		}
+
+		HorrorUI->SetupCharacter(HorrorCharacter);
+	}
+}
+
+void AHorrorPlayerController::SetupInputComponent()
+{
+	Super::SetupInputComponent();
+	
+	// Add Input Mapping Contexts
+	if (UEnhancedInputLocalPlayerSubsystem* Subsystem = ULocalPlayer::GetSubsystem<UEnhancedInputLocalPlayerSubsystem>(GetLocalPlayer()))
+	{
+		for (UInputMappingContext* CurrentContext : DefaultMappingContexts)
+		{
+			Subsystem->AddMappingContext(CurrentContext, 0);
+		}
+	}
+}
diff --git a/Templates/TP_FirstPerson/Source/TP_FirstPerson/Variant_Horror/HorrorPlayerController.h b/Templates/TP_FirstPerson/Source/TP_FirstPerson/Variant_Horror/HorrorPlayerController.h
new file mode 100644
--- /dev/null
+++ b/Templates/TP_FirstPerson/Source/TP_FirstPerson/Variant_Horror/HorrorPlayerController.h
@@ -0,0 +1,48 @@
+// Copyright Epic Games, Inc. All Rights Reserved.
+
+#pragma once
+
+#include "CoreMinimal.h"
+#include "GameFramework/PlayerController.h"
+#include "HorrorPlayerController.generated.h"
+
+class UInputMappingContext;
+class UHorrorUI;
+
+/**
+ *  Player Controller for a first person horror game
+ *  Manages input mappings
+ *  Manages UI
+ */
+UCLASS(abstract)
+class TP_FIRSTPERSON_API AHorrorPlayerController : public APlayerController
+{
+	GENERATED_BODY()
+	
+protected:
+
+	/** Type of UI widget to spawn */
+	UPROPERTY(EditAnywhere, Category="Horror")
+	TSubclassOf<UHorrorUI> HorrorUIClass;
+
+	/** Pointer to the UI widget */
+	TObjectPtr<UHorrorUI> HorrorUI;
+
+public:
+
+	/** Constructor */
+	AHorrorPlayerController();
+
+protected:
+
+	/** Possessed pawn initialization */
+	virtual void OnPossess(APawn* aPawn) override;
+
+	/** Input Mapping Contexts */
+	UPROPERTY(EditAnywhere, BlueprintReadOnly, Category ="Input", meta = (AllowPrivateAccess = "true"))
+	TArray<UInputMappingContext*> DefaultMappingContexts;
+
+	/** Input mapping context setup */
+	virtual void SetupInputComponent() override;
+
+};
diff --git a/Templates/TP_FirstPerson/Source/TP_FirstPerson/Variant_Horror/HorrorUI.cpp b/Templates/TP_FirstPerson/Source/TP_FirstPerson/Variant_Horror/HorrorUI.cpp
new file mode 100644
--- /dev/null
+++ b/Templates/TP_FirstPerson/Source/TP_FirstPerson/Variant_Horror/HorrorUI.cpp
@@ -0,0 +1,23 @@
+// Copyright Epic Games, Inc. All Rights Reserved.
+
+
+#include "Variant_Horror/HorrorUI.h"
+#include "HorrorCharacter.h"
+
+void UHorrorUI::SetupCharacter(AHorrorCharacter* HorrorCharacter)
+{
+	HorrorCharacter->OnSprintMeterUpdated.AddDynamic(this, &UHorrorUI::OnSprintMeterUpdated);
+	HorrorCharacter->OnSprintStateChanged.AddDynamic(this, &UHorrorUI::OnSprintStateChanged);
+}
+
+void UHorrorUI::OnSprintMeterUpdated(float Percent)
+{
+	// call the BP handler
+	BP_SprintMeterUpdated(Percent);
+}
+
+void UHorrorUI::OnSprintStateChanged(bool bSprinting)
+{
+	// call the BP handler
+	BP_SprintStateChanged(bSprinting);
+}
diff --git a/Templates/TP_FirstPerson/Source/TP_FirstPerson/Variant_Horror/HorrorUI.h b/Templates/TP_FirstPerson/Source/TP_FirstPerson/Variant_Horror/HorrorUI.h
new file mode 100644
--- /dev/null
+++ b/Templates/TP_FirstPerson/Source/TP_FirstPerson/Variant_Horror/HorrorUI.h
@@ -0,0 +1,42 @@
+// Copyright Epic Games, Inc. All Rights Reserved.
+
+#pragma once
+
+#include "CoreMinimal.h"
+#include "Blueprint/UserWidget.h"
+#include "HorrorUI.generated.h"
+
+class AHorrorCharacter;
+
+/**
+ *  Simple UI for a first person horror game
+ *  Manages character sprint meter display
+ */
+UCLASS(abstract)
+class TP_FIRSTPERSON_API UHorrorUI : public UUserWidget
+{
+	GENERATED_BODY()
+	
+public:
+
+	/** Sets up delegate listeners for the passed character */
+	void SetupCharacter(AHorrorCharacter* HorrorCharacter);
+
+	/** Called when the character's sprint meter is updated */
+	UFUNCTION()
+	void OnSprintMeterUpdated(float Percent);
+
+	/** Called when the character's sprint state changes */
+	UFUNCTION()
+	void OnSprintStateChanged(bool bSprinting);
+
+protected:
+
+	/** Passes control to Blueprint to update the sprint meter widgets */
+	UFUNCTION(BlueprintImplementableEvent, Category="Horror", meta=(DisplayName = "Sprint Meter Updated"))
+	void BP_SprintMeterUpdated(float Percent);
+
+	/** Passes control to Blueprint to update the sprint meter status */
+	UFUNCTION(BlueprintImplementableEvent, Category="Horror", meta=(DisplayName = "Sprint State Changed"))
+	void BP_SprintStateChanged(bool bSprinting);
+};
diff --git a/Templates/TP_FirstPerson/Source/TP_FirstPerson/Variant_Shooter/AI/EnvQueryContext_Target.cpp b/Templates/TP_FirstPerson/Source/TP_FirstPerson/Variant_Shooter/AI/EnvQueryContext_Target.cpp
new file mode 100644
--- /dev/null
+++ b/Templates/TP_FirstPerson/Source/TP_FirstPerson/Variant_Shooter/AI/EnvQueryContext_Target.cpp
@@ -0,0 +1,22 @@
+// Copyright Epic Games, Inc. All Rights Reserved.
+
+
+#include "Variant_Shooter/AI/EnvQueryContext_Target.h"
+#include "EnvironmentQuery/Items/EnvQueryItemType_Actor.h"
+#include "EnvironmentQuery/EnvQueryTypes.h"
+#include "ShooterAIController.h"
+
+void UEnvQueryContext_Target::ProvideContext(FEnvQueryInstance& QueryInstance, FEnvQueryContextData& ContextData) const
+{
+	// get the controller from the query instance
+	if (AShooterAIController* Controller = Cast<AShooterAIController>(QueryInstance.Owner))
+	{
+		// ensure the target is valid
+		if (IsValid(Controller->GetCurrentTarget()))
+		{
+			// add the controller's target actor to the context
+			UEnvQueryItemType_Actor::SetContextHelper(ContextData, Controller->GetCurrentTarget());
+		}
+	}
+
+}
diff --git a/Templates/TP_FirstPerson/Source/TP_FirstPerson/Variant_Shooter/AI/EnvQueryContext_Target.h b/Templates/TP_FirstPerson/Source/TP_FirstPerson/Variant_Shooter/AI/EnvQueryContext_Target.h
new file mode 100644
--- /dev/null
+++ b/Templates/TP_FirstPerson/Source/TP_FirstPerson/Variant_Shooter/AI/EnvQueryContext_Target.h
@@ -0,0 +1,22 @@
+// Copyright Epic Games, Inc. All Rights Reserved.
+
+#pragma once
+
+#include "CoreMinimal.h"
+#include "EnvironmentQuery/EnvQueryContext.h"
+#include "EnvQueryContext_Target.generated.h"
+
+/**
+ *  Custom EnvQuery Context that returns the actor currently targeted by an NPC
+ */
+UCLASS()
+class TP_FIRSTPERSON_API UEnvQueryContext_Target : public UEnvQueryContext
+{
+	GENERATED_BODY()
+	
+public:
+
+	/** Provides the context locations or actors for this EnvQuery */
+	virtual void ProvideContext(FEnvQueryInstance& QueryInstance, FEnvQueryContextData& ContextData) const override;
+
+};
diff --git a/Templates/TP_FirstPerson/Source/TP_FirstPerson/Variant_Shooter/AI/ShooterAIController.cpp b/Templates/TP_FirstPerson/Source/TP_FirstPerson/Variant_Shooter/AI/ShooterAIController.cpp
new file mode 100644
--- /dev/null
+++ b/Templates/TP_FirstPerson/Source/TP_FirstPerson/Variant_Shooter/AI/ShooterAIController.cpp
@@ -0,0 +1,74 @@
+// Copyright Epic Games, Inc. All Rights Reserved.
+
+
+#include "Variant_Shooter/AI/ShooterAIController.h"
+#include "ShooterNPC.h"
+#include "Components/StateTreeAIComponent.h"
+#include "Perception/AIPerceptionComponent.h"
+#include "Navigation/PathFollowingComponent.h"
+#include "AI/Navigation/PathFollowingAgentInterface.h"
+
+AShooterAIController::AShooterAIController()
+{
+	// create the StateTree component
+	StateTreeAI = CreateDefaultSubobject<UStateTreeAIComponent>(TEXT("StateTreeAI"));
+
+	// create the AI perception component. It will be configured in BP
+	AIPerception = CreateDefaultSubobject<UAIPerceptionComponent>(TEXT("AIPerception"));
+
+	// subscribe to the AI perception delegates
+	AIPerception->OnTargetPerceptionUpdated.AddDynamic(this, &AShooterAIController::OnPerceptionUpdated);
+	AIPerception->OnTargetPerceptionForgotten.AddDynamic(this, &AShooterAIController::OnPerceptionForgotten);
+}
+
+void AShooterAIController::OnPossess(APawn* InPawn)
+{
+	Super::OnPossess(InPawn);
+
+	// ensure we're possessing an NPC
+	if (AShooterNPC* NPC = Cast<AShooterNPC>(InPawn))
+	{
+		// add the team tag to the pawn
+		NPC->Tags.Add(TeamTag);
+
+		// subscribe to the pawn's OnDeath delegate
+		NPC->OnPawnDeath.AddDynamic(this, &AShooterAIController::OnPawnDeath);
+	}
+}
+
+void AShooterAIController::OnPawnDeath()
+{
+	// stop movement
+	GetPathFollowingComponent()->AbortMove(*this, FPathFollowingResultFlags::UserAbort);
+
+	// stop StateTree logic
+	StateTreeAI->StopLogic(FString(""));
+
+	// unpossess the pawn
+	UnPossess();
+
+	// destroy this controller
+	Destroy();
+}
+
+void AShooterAIController::SetCurrentTarget(AActor* Target)
+{
+	TargetEnemy = Target;
+}
+
+void AShooterAIController::ClearCurrentTarget()
+{
+	TargetEnemy = nullptr;
+}
+
+void AShooterAIController::OnPerceptionUpdated(AActor* Actor, FAIStimulus Stimulus)
+{
+	// pass the data to the StateTree delegate hook
+	OnShooterPerceptionUpdated.ExecuteIfBound(Actor, Stimulus);
+}
+
+void AShooterAIController::OnPerceptionForgotten(AActor* Actor)
+{
+	// pass the data to the StateTree delegate hook
+	OnShooterPerceptionForgotten.ExecuteIfBound(Actor);
+}
diff --git a/Templates/TP_FirstPerson/Source/TP_FirstPerson/Variant_Shooter/AI/ShooterAIController.h b/Templates/TP_FirstPerson/Source/TP_FirstPerson/Variant_Shooter/AI/ShooterAIController.h
new file mode 100644
--- /dev/null
+++ b/Templates/TP_FirstPerson/Source/TP_FirstPerson/Variant_Shooter/AI/ShooterAIController.h
@@ -0,0 +1,85 @@
+// Copyright Epic Games, Inc. All Rights Reserved.
+
+#pragma once
+
+#include "CoreMinimal.h"
+#include "AIController.h"
+#include "ShooterAIController.generated.h"
+
+class UStateTreeAIComponent;
+class UAIPerceptionComponent;
+struct FAIStimulus;
+
+DECLARE_DELEGATE_TwoParams(FShooterPerceptionUpdatedDelegate, AActor*, const FAIStimulus&);
+DECLARE_DELEGATE_OneParam(FShooterPerceptionForgottenDelegate, AActor*);
+
+/**
+ *  Simple AI Controller for a first person shooter enemy
+ */
+UCLASS(abstract)
+class TP_FIRSTPERSON_API AShooterAIController : public AAIController
+{
+	GENERATED_BODY()
+	
+	/** Runs the behavior StateTree for this NPC */
+	UPROPERTY(VisibleAnywhere, BlueprintReadOnly, Category="Components", meta = (AllowPrivateAccess = "true"))
+	UStateTreeAIComponent* StateTreeAI;
+
+	/** Detects other actors through sight, hearing and other senses */
+	UPROPERTY(VisibleAnywhere, BlueprintReadOnly, Category="Components", meta = (AllowPrivateAccess = "true"))
+	UAIPerceptionComponent* AIPerception;
+
+protected:
+
+	/** Team tag for pawn friend or foe identification */
+	UPROPERTY(EditAnywhere, Category="Shooter")
+	FName TeamTag = FName("Enemy");
+
+	/** Enemy currently being targeted */
+	TObjectPtr<AActor> TargetEnemy;
+
+public:
+
+	/** Called when an AI perception has been updated. StateTree task delegate hook */
+	FShooterPerceptionUpdatedDelegate OnShooterPerceptionUpdated;
+
+	/** Called when an AI perception has been forgotten. StateTree task delegate hook */
+	FShooterPerceptionForgottenDelegate OnShooterPerceptionForgotten;
+
+public:
+
+	/** Constructor */
+	AShooterAIController();
+
+protected:
+
+	/** Pawn initialization */
+	virtual void OnPossess(APawn* InPawn) override;
+
+protected:
+
+	/** Called when the possessed pawn dies */
+	UFUNCTION()
+	void OnPawnDeath();
+
+public:
+
+	/** Sets the targeted enemy */
+	void SetCurrentTarget(AActor* Target);
+
+	/** Clears the targeted enemy */
+	void ClearCurrentTarget();
+
+	/** Returns the targeted enemy */
+	AActor* GetCurrentTarget() const { return TargetEnemy; };
+
+protected:
+
+	/** Called when the AI perception component updates a perception on a given actor */
+	UFUNCTION()
+	void OnPerceptionUpdated(AActor* Actor, FAIStimulus Stimulus);
+
+	/** Called when the AI perception component forgets a given actor */
+	UFUNCTION()
+	void OnPerceptionForgotten(AActor* Actor);
+};
diff --git a/Templates/TP_FirstPerson/Source/TP_FirstPerson/Variant_Shooter/AI/ShooterNPC.cpp b/Templates/TP_FirstPerson/Source/TP_FirstPerson/Variant_Shooter/AI/ShooterNPC.cpp
new file mode 100644
--- /dev/null
+++ b/Templates/TP_FirstPerson/Source/TP_FirstPerson/Variant_Shooter/AI/ShooterNPC.cpp
@@ -0,0 +1,208 @@
+// Copyright Epic Games, Inc. All Rights Reserved.
+
+
+#include "Variant_Shooter/AI/ShooterNPC.h"
+#include "ShooterWeapon.h"
+#include "Components/SkeletalMeshComponent.h"
+#include "Camera/CameraComponent.h"
+#include "Kismet/KismetMathLibrary.h"
+#include "Engine/World.h"
+#include "ShooterGameMode.h"
+#include "Components/CapsuleComponent.h"
+#include "GameFramework/CharacterMovementComponent.h"
+#include "TimerManager.h"
+
+void AShooterNPC::BeginPlay()
+{
+	Super::BeginPlay();
+
+	// spawn the weapon
+	FActorSpawnParameters SpawnParams;
+	SpawnParams.Owner = this;
+	SpawnParams.Instigator = this;
+	SpawnParams.SpawnCollisionHandlingOverride = ESpawnActorCollisionHandlingMethod::AlwaysSpawn;
+
+	Weapon = GetWorld()->SpawnActor<AShooterWeapon>(WeaponClass, GetActorTransform(), SpawnParams);
+}
+
+void AShooterNPC::EndPlay(const EEndPlayReason::Type EndPlayReason)
+{
+	Super::EndPlay(EndPlayReason);
+
+	// clear the death timer
+	GetWorld()->GetTimerManager().ClearTimer(DeathTimer);
+}
+
+float AShooterNPC::TakeDamage(float Damage, struct FDamageEvent const& DamageEvent, AController* EventInstigator, AActor* DamageCauser)
+{
+	// ignore if already dead
+	if (bIsDead)
+	{
+		return 0.0f;
+	}
+
+	// Reduce HP
+	CurrentHP -= Damage;
+
+	// Have we depleted HP?
+	if (CurrentHP <= 0.0f)
+	{
+		Die();
+	}
+
+	return Damage;
+}
+
+void AShooterNPC::AttachWeaponMeshes(AShooterWeapon* WeaponToAttach)
+{
+	const FAttachmentTransformRules AttachmentRule(EAttachmentRule::SnapToTarget, false);
+
+	// attach the weapon actor
+	WeaponToAttach->AttachToActor(this, AttachmentRule);
+
+	// attach the weapon meshes
+	WeaponToAttach->GetFirstPersonMesh()->AttachToComponent(GetFirstPersonMesh(), AttachmentRule, FirstPersonWeaponSocket);
+	WeaponToAttach->GetThirdPersonMesh()->AttachToComponent(GetMesh(), AttachmentRule, FirstPersonWeaponSocket);
+}
+
+void AShooterNPC::PlayFiringMontage(UAnimMontage* Montage)
+{
+	// unused
+}
+
+void AShooterNPC::AddWeaponRecoil(float Recoil)
+{
+	// unused
+}
+
+void AShooterNPC::UpdateWeaponHUD(int32 CurrentAmmo, int32 MagazineSize)
+{
+	// unused
+}
+
+FVector AShooterNPC::GetWeaponTargetLocation()
+{
+	// start aiming from the camera location
+	const FVector AimSource = GetFirstPersonCameraComponent()->GetComponentLocation();
+
+	FVector AimDir, AimTarget = FVector::ZeroVector;
+
+	// do we have an aim target?
+	if (CurrentAimTarget)
+	{
+		// target the actor location
+		AimTarget = CurrentAimTarget->GetActorLocation();
+
+		// apply a vertical offset to target head/feet
+		AimTarget.Z += FMath::RandRange(MinAimOffsetZ, MaxAimOffsetZ);
+
+		// get the aim direction and apply randomness in a cone
+		AimDir = (AimTarget - AimSource).GetSafeNormal();
+		AimDir = UKismetMathLibrary::RandomUnitVectorInConeInDegrees(AimDir, AimVarianceHalfAngle);
+
+		
+	} else {
+
+		// no aim target, so just use the camera facing
+		AimDir = UKismetMathLibrary::RandomUnitVectorInConeInDegrees(GetFirstPersonCameraComponent()->GetForwardVector(), AimVarianceHalfAngle);
+
+	}
+
+	// calculate the unobstructed aim target location
+	AimTarget = AimSource + (AimDir * AimRange);
+
+	// run a visibility trace to see if there's obstructions
+	FHitResult OutHit;
+
+	FCollisionQueryParams QueryParams;
+	QueryParams.AddIgnoredActor(this);
+
+	GetWorld()->LineTraceSingleByChannel(OutHit, AimSource, AimTarget, ECC_Visibility, QueryParams);
+
+	// return either the impact point or the trace end
+	return OutHit.bBlockingHit ? OutHit.ImpactPoint : OutHit.TraceEnd;
+}
+
+void AShooterNPC::AddWeaponClass(const TSubclassOf<AShooterWeapon>& InWeaponClass)
+{
+	// unused
+}
+
+void AShooterNPC::OnWeaponActivated(AShooterWeapon* InWeapon)
+{
+	// unused
+}
+
+void AShooterNPC::OnWeaponDeactivated(AShooterWeapon* InWeapon)
+{
+	// unused
+}
+
+void AShooterNPC::OnSemiWeaponRefire()
+{
+	// are we still shooting?
+	if (bIsShooting)
+	{
+		// fire the weapon
+		Weapon->StartFiring();
+	}
+}
+
+void AShooterNPC::Die()
+{
+	// ignore if already dead
+	if (bIsDead)
+	{
+		return;
+	}
+
+	// raise the dead flag
+	bIsDead = true;
+
+	// increment the team score
+	if (AShooterGameMode* GM = Cast<AShooterGameMode>(GetWorld()->GetAuthGameMode()))
+	{
+		GM->IncrementTeamScore(TeamByte);
+	}
+
+	// disable capsule collision
+	GetCapsuleComponent()->SetCollisionEnabled(ECollisionEnabled::NoCollision);
+
+	// stop movement
+	GetCharacterMovement()->StopMovementImmediately();
+	GetCharacterMovement()->StopActiveMovement();
+
+	// enable ragdoll physics on the third person mesh
+	GetMesh()->SetCollisionProfileName(RagdollCollisionProfile);
+	GetMesh()->SetSimulatePhysics(true);
+	GetMesh()->SetPhysicsBlendWeight(1.0f);
+
+	// schedule actor destruction
+	GetWorld()->GetTimerManager().SetTimer(DeathTimer, this, &AShooterNPC::DeferredDestruction, DeferredDestructionTime, false);
+}
+
+void AShooterNPC::DeferredDestruction()
+{
+	Destroy();
+}
+
+void AShooterNPC::StartShooting(AActor* ActorToShoot)
+{
+	// save the aim target
+	CurrentAimTarget = ActorToShoot;
+
+	// raise the flag
+	bIsShooting = true;
+
+	// signal the weapon
+	Weapon->StartFiring();
+}
+
+void AShooterNPC::StopShooting()
+{
+	// lower the flag
+	bIsShooting = false;
+
+	// signal the weapon
+	Weapon->StopFiring();
+}
diff --git a/Templates/TP_FirstPerson/Source/TP_FirstPerson/Variant_Shooter/AI/ShooterNPC.h b/Templates/TP_FirstPerson/Source/TP_FirstPerson/Variant_Shooter/AI/ShooterNPC.h
new file mode 100644
--- /dev/null
+++ b/Templates/TP_FirstPerson/Source/TP_FirstPerson/Variant_Shooter/AI/ShooterNPC.h
@@ -0,0 +1,153 @@
+// Copyright Epic Games, Inc. All Rights Reserved.
+
+#pragma once
+
+#include "CoreMinimal.h"
+#include "TP_FirstPersonCharacter.h"
+#include "ShooterWeaponHolder.h"
+#include "ShooterNPC.generated.h"
+
+DECLARE_DYNAMIC_MULTICAST_DELEGATE(FPawnDeathDelegate);
+
+class AShooterWeapon;
+
+/**
+ *  A simple AI-controlled shooter game NPC
+ *  Executes its behavior through a StateTree managed by its AI Controller
+ *  Holds and manages a weapon
+ */
+UCLASS(abstract)
+class TP_FIRSTPERSON_API AShooterNPC : public ATP_FirstPersonCharacter, public IShooterWeaponHolder
+{
+	GENERATED_BODY()
+
+public:
+
+	/** Current HP for this character. It dies if it reaches zero through damage */
+	UPROPERTY(EditAnywhere, BlueprintReadOnly, Category="Damage")
+	float CurrentHP = 100.0f;
+
+protected:
+
+	/** Name of the collision profile to use during ragdoll death */
+	UPROPERTY(EditAnywhere, Category="Damage")
+	FName RagdollCollisionProfile = FName("Ragdoll");
+
+	/** Time to wait after death before destroying this actor */
+	UPROPERTY(EditAnywhere, Category="Damage")
+	float DeferredDestructionTime = 5.0f;
+
+	/** Team byte for this character */
+	UPROPERTY(EditAnywhere, Category="Team")
+	uint8 TeamByte = 1;
+
+	/** Pointer to the equipped weapon */
+	TObjectPtr<AShooterWeapon> Weapon;
+
+	/** Type of weapon to spawn for this character */
+	UPROPERTY(EditAnywhere, Category="Weapon")
+	TSubclassOf<AShooterWeapon> WeaponClass;
+
+	/** Name of the first person mesh weapon socket */
+	UPROPERTY(EditAnywhere, BlueprintReadOnly, Category ="Weapons")
+	FName FirstPersonWeaponSocket = FName("HandGrip_R");
+
+	/** Name of the third person mesh weapon socket */
+	UPROPERTY(EditAnywhere, BlueprintReadOnly, Category ="Weapons")
+	FName ThirdPersonWeaponSocket = FName("HandGrip_R");
+
+	/** Max range for aiming calculations */
+	UPROPERTY(EditAnywhere, Category="Aim")
+	float AimRange = 10000.0f;
+
+	/** Cone variance to apply while aiming */
+	UPROPERTY(EditAnywhere, Category="Aim")
+	float AimVarianceHalfAngle = 10.0f;
+
+	/** Minimum vertical offset from the target center to apply when aiming */
+	UPROPERTY(EditAnywhere, Category="Aim")
+	float MinAimOffsetZ = -35.0f;
+
+	/** Maximum vertical offset from the target center to apply when aiming */
+	UPROPERTY(EditAnywhere, Category="Aim")
+	float MaxAimOffsetZ = -60.0f;
+
+	/** Actor currently being targeted */
+	TObjectPtr<AActor> CurrentAimTarget;
+
+	/** If true, this character is currently shooting its weapon */
+	bool bIsShooting = false;
+
+	/** If true, this character has already died */
+	bool bIsDead = false;
+
+	/** Deferred destruction on death timer */
+	FTimerHandle DeathTimer;
+
+public:
+
+	/** Delegate called when this NPC dies */
+	FPawnDeathDelegate OnPawnDeath;
+
+protected:
+
+	/** Gameplay initialization */
+	virtual void BeginPlay() override;
+
+	/** Gameplay cleanup */
+	virtual void EndPlay(const EEndPlayReason::Type EndPlayReason) override;
+
+public:
+
+	/** Handle incoming damage */
+	virtual float TakeDamage(float Damage, struct FDamageEvent const& DamageEvent, AController* EventInstigator, AActor* DamageCauser) override;
+
+public:
+
+	//~Begin IShooterWeaponHolder interface
+
+	/** Attaches a weapon's meshes to the owner */
+	virtual void AttachWeaponMeshes(AShooterWeapon* Weapon) override;
+
+	/** Plays the firing montage for the weapon */
+	virtual void PlayFiringMontage(UAnimMontage* Montage) override;
+
+	/** Applies weapon recoil to the owner */
+	virtual void AddWeaponRecoil(float Recoil) override;
+
+	/** Updates the weapon's HUD with the current ammo count */
+	virtual void UpdateWeaponHUD(int32 CurrentAmmo, int32 MagazineSize) override;
+
+	/** Calculates and returns the aim location for the weapon */
+	virtual FVector GetWeaponTargetLocation() override;
+
+	/** Gives a weapon of this class to the owner */
+	virtual void AddWeaponClass(const TSubclassOf<AShooterWeapon>& WeaponClass) override;
+
+	/** Activates the passed weapon */
+	virtual void OnWeaponActivated(AShooterWeapon* Weapon) override;
+
+	/** Deactivates the passed weapon */
+	virtual void OnWeaponDeactivated(AShooterWeapon* Weapon) override;
+
+	/** Notifies the owner that the weapon cooldown has expired and it's ready to shoot again */
+	virtual void OnSemiWeaponRefire() override;
+
+	//~End IShooterWeaponHolder interface
+
+protected:
+
+	/** Called when HP is depleted and the character should die */
+	void Die();
+
+	/** Called after death to destroy the actor */
+	void DeferredDestruction();
+
+public:
+
+	/** Signals this character to start shooting at the passed actor */
+	void StartShooting(AActor* ActorToShoot);
+
+	/** Signals this character to stop shooting */
+	void StopShooting();
+};
diff --git a/Templates/TP_FirstPerson/Source/TP_FirstPerson/Variant_Shooter/AI/ShooterStateTreeUtility.cpp b/Templates/TP_FirstPerson/Source/TP_FirstPerson/Variant_Shooter/AI/ShooterStateTreeUtility.cpp
new file mode 100644
--- /dev/null
+++ b/Templates/TP_FirstPerson/Source/TP_FirstPerson/Variant_Shooter/AI/ShooterStateTreeUtility.cpp
@@ -0,0 +1,370 @@
+// Copyright Epic Games, Inc. All Rights Reserved.
+
+
+#include "Variant_Shooter/AI/ShooterStateTreeUtility.h"
+#include "StateTreeExecutionContext.h"
+#include "ShooterNPC.h"
+#include "Camera/CameraComponent.h"
+#include "AIController.h"
+#include "Perception/AIPerceptionComponent.h"
+#include "ShooterAIController.h"
+#include "StateTreeAsyncExecutionContext.h"
+
+#define LOCTEXT_NAMESPACE "FirstPersonTemplate"
+
+bool FStateTreeLineOfSightToTargetCondition::TestCondition(FStateTreeExecutionContext& Context) const
+{
+	const FInstanceDataType& InstanceData = Context.GetInstanceData(*this);
+
+	// ensure the target is valid
+	if (!IsValid(InstanceData.Target))
+	{
+		return !InstanceData.bMustHaveLineOfSight;
+	}
+	
+	// check if the character is facing towards the target
+	const FVector TargetDir = (InstanceData.Target->GetActorLocation() - InstanceData.Character->GetActorLocation()).GetSafeNormal();
+
+	const float FacingDot = FVector::DotProduct(TargetDir, InstanceData.Character->GetActorForwardVector());
+	const float MaxDot = FMath::Cos(FMath::DegreesToRadians(InstanceData.LineOfSightConeAngle));
+
+	// is the facing outside of our cone half angle?
+	if (FacingDot <= MaxDot)
+	{
+		return !InstanceData.bMustHaveLineOfSight;
+	}
+
+	// get the target's bounding box
+	FVector CenterOfMass, Extent;
+	InstanceData.Target->GetActorBounds(true, CenterOfMass, Extent, false);
+
+	// divide the vertical extent by the number of line of sight checks we'll do
+	const float ExtentZOffset = Extent.Z * 2.0f / InstanceData.NumberOfVerticalLineOfSightChecks;
+
+	// get the character's camera location as the source for the line checks
+	const FVector Start = InstanceData.Character->GetFirstPersonCameraComponent()->GetComponentLocation();
+
+	// ignore the character and target. We want to ensure there's an unobstructed trace not counting them
+	FCollisionQueryParams QueryParams;
+	QueryParams.AddIgnoredActor(InstanceData.Character);
+	QueryParams.AddIgnoredActor(InstanceData.Target);
+
+	FHitResult OutHit;
+
+	// run a number of vertically offset line traces to the target location
+	for (int32 i = 0; i < InstanceData.NumberOfVerticalLineOfSightChecks - 1; ++i)
+	{
+		// calculate the endpoint for the trace
+		const FVector End = CenterOfMass + FVector(0.0f, 0.0f, Extent.Z - ExtentZOffset * i);
+
+		InstanceData.Character->GetWorld()->LineTraceSingleByChannel(OutHit, Start, End, ECC_Visibility, QueryParams);
+
+		// is the trace unobstructed?
+		if (!OutHit.bBlockingHit)
+		{
+			// we only need one unobstructed trace, so terminate early
+			return InstanceData.bMustHaveLineOfSight;
+		}
+	}
+
+	// no line of sight found
+	return !InstanceData.bMustHaveLineOfSight;
+}
+
+#if WITH_EDITOR
+FText FStateTreeLineOfSightToTargetCondition::GetDescription(const FGuid& ID, FStateTreeDataView InstanceDataView, const IStateTreeBindingLookup& BindingLookup, EStateTreeNodeFormatting Formatting /*= EStateTreeNodeFormatting::Text*/) const
+{
+	return LOCTEXT("StateTreeConditionLineOfSightToTarget", "<b>Has Line of Sight</b>");
+}
+#endif
+
+////////////////////////////////////////////////////////////////////
+
+EStateTreeRunStatus FStateTreeFaceActorTask::EnterState(FStateTreeExecutionContext& Context, const FStateTreeTransitionResult& Transition) const
+{
+	// have we transitioned from another state?
+	if (Transition.ChangeType == EStateTreeStateChangeType::Changed)
+	{
+		// get the instance data
+		FInstanceDataType& InstanceData = Context.GetInstanceData(*this);
+
+		// set the AI Controller's focus
+		InstanceData.Controller->SetFocus(InstanceData.ActorToFaceTowards);
+	}
+
+	return EStateTreeRunStatus::Running;
+}
+
+void FStateTreeFaceActorTask::ExitState(FStateTreeExecutionContext& Context, const FStateTreeTransitionResult& Transition) const
+{
+	// have we transitioned to another state?
+	if (Transition.ChangeType == EStateTreeStateChangeType::Changed)
+	{
+		// get the instance data
+		FInstanceDataType& InstanceData = Context.GetInstanceData(*this);
+
+		// clear the AI Controller's focus
+		InstanceData.Controller->ClearFocus(EAIFocusPriority::Gameplay);
+	}
+}
+
+#if WITH_EDITOR
+FText FStateTreeFaceActorTask::GetDescription(const FGuid& ID, FStateTreeDataView InstanceDataView, const IStateTreeBindingLookup& BindingLookup, EStateTreeNodeFormatting Formatting /*= EStateTreeNodeFormatting::Text*/) const
+{
+	return LOCTEXT("StateTreeTaskFaceActorDescription", "<b>Face Towards Actor</b>");
+}
+#endif // WITH_EDITOR
+
+////////////////////////////////////////////////////////////////////
+
+EStateTreeRunStatus FStateTreeFaceLocationTask::EnterState(FStateTreeExecutionContext& Context, const FStateTreeTransitionResult& Transition) const
+{
+	// have we transitioned from another state?
+	if (Transition.ChangeType == EStateTreeStateChangeType::Changed)
+	{
+		// get the instance data
+		FInstanceDataType& InstanceData = Context.GetInstanceData(*this);
+
+		// set the AI Controller's focus
+		InstanceData.Controller->SetFocalPoint(InstanceData.FaceLocation);
+	}
+
+	return EStateTreeRunStatus::Running;
+}
+
+void FStateTreeFaceLocationTask::ExitState(FStateTreeExecutionContext& Context, const FStateTreeTransitionResult& Transition) const
+{
+	// have we transitioned to another state?
+	if (Transition.ChangeType == EStateTreeStateChangeType::Changed)
+	{
+		// get the instance data
+		FInstanceDataType& InstanceData = Context.GetInstanceData(*this);
+
+		// clear the AI Controller's focus
+		InstanceData.Controller->ClearFocus(EAIFocusPriority::Gameplay);
+	}
+}
+
+#if WITH_EDITOR
+FText FStateTreeFaceLocationTask::GetDescription(const FGuid& ID, FStateTreeDataView InstanceDataView, const IStateTreeBindingLookup& BindingLookup, EStateTreeNodeFormatting Formatting /*= EStateTreeNodeFormatting::Text*/) const
+{
+	return LOCTEXT("StateTreeTaskFaceLocationDescription", "<b>Face Towards Location</b>");
+}
+#endif // WITH_EDITOR
+
+////////////////////////////////////////////////////////////////////
+
+EStateTreeRunStatus FStateTreeSetRandomFloatTask::EnterState(FStateTreeExecutionContext& Context, const FStateTreeTransitionResult& Transition) const
+{
+	// have we transitioned to another state?
+	if (Transition.ChangeType == EStateTreeStateChangeType::Changed)
+	{
+		// get the instance data
+		FInstanceDataType& InstanceData = Context.GetInstanceData(*this);
+
+		// calculate the output value
+		InstanceData.OutValue = FMath::RandRange(InstanceData.MinValue, InstanceData.MaxValue);
+	}
+
+	return EStateTreeRunStatus::Running;
+}
+
+#if WITH_EDITOR
+FText FStateTreeSetRandomFloatTask::GetDescription(const FGuid& ID, FStateTreeDataView InstanceDataView, const IStateTreeBindingLookup& BindingLookup, EStateTreeNodeFormatting Formatting /*= EStateTreeNodeFormatting::Text*/) const
+{
+	return LOCTEXT("StateTreeTaskSetRandomFloatDescription", "<b>Set Random Float</b>");
+}
+#endif // WITH_EDITOR
+
+////////////////////////////////////////////////////////////////////
+
+EStateTreeRunStatus FStateTreeShootAtTargetTask::EnterState(FStateTreeExecutionContext& Context, const FStateTreeTransitionResult& Transition) const
+{
+	// have we transitioned from another state?
+	if (Transition.ChangeType == EStateTreeStateChangeType::Changed)
+	{
+		// get the instance data
+		FInstanceDataType& InstanceData = Context.GetInstanceData(*this);
+
+		// tell the character to shoot the target
+		InstanceData.Character->StartShooting(InstanceData.Target);
+	}
+
+	return EStateTreeRunStatus::Running;
+}
+
+void FStateTreeShootAtTargetTask::ExitState(FStateTreeExecutionContext& Context, const FStateTreeTransitionResult& Transition) const
+{
+	// have we transitioned to another state?
+	if (Transition.ChangeType == EStateTreeStateChangeType::Changed)
+	{
+		// get the instance data
+		FInstanceDataType& InstanceData = Context.GetInstanceData(*this);
+
+		// tell the character to stop shooting
+		InstanceData.Character->StopShooting();
+	}
+}
+
+#if WITH_EDITOR
+FText FStateTreeShootAtTargetTask::GetDescription(const FGuid& ID, FStateTreeDataView InstanceDataView, const IStateTreeBindingLookup& BindingLookup, EStateTreeNodeFormatting Formatting /*= EStateTreeNodeFormatting::Text*/) const
+{
+	return LOCTEXT("StateTreeTaskShootAtTargetDescription", "<b>Shoot at Target</b>");
+}
+#endif // WITH_EDITOR
+
+EStateTreeRunStatus FStateTreeSenseEnemiesTask::EnterState(FStateTreeExecutionContext& Context, const FStateTreeTransitionResult& Transition) const
+{
+	// have we transitioned from another state?
+	if (Transition.ChangeType == EStateTreeStateChangeType::Changed)
+	{
+		// get the instance data
+		FInstanceDataType& InstanceData = Context.GetInstanceData(*this);
+
+		// bind the perception updated delegate on the controller
+		InstanceData.Controller->OnShooterPerceptionUpdated.BindLambda(
+			[WeakContext = Context.MakeWeakExecutionContext()](AActor* SensedActor, const FAIStimulus& Stimulus)
+			{
+				// get the instance data inside the lambda
+				FInstanceDataType* LambdaInstanceData = WeakContext.MakeStrongExecutionContext().GetInstanceDataPtr<FInstanceDataType>();
+
+				if (!LambdaInstanceData)
+				{
+					return;
+				}
+
+				if (SensedActor->ActorHasTag(LambdaInstanceData->SenseTag))
+				{
+					bool bDirectLOS = false;
+
+					// calculate the direction of the stimulus
+					const FVector StimulusDir = (Stimulus.StimulusLocation - LambdaInstanceData->Character->GetActorLocation()).GetSafeNormal();
+
+					// infer the angle from the dot product between the character facing and the stimulus direction
+					const float DirDot = FVector::DotProduct(StimulusDir, LambdaInstanceData->Character->GetActorForwardVector());
+					const float MaxDot = FMath::Cos(FMath::DegreesToRadians(LambdaInstanceData->DirectLineOfSightCone));
+
+					// is the direction within our perception cone?
+					if (DirDot >= MaxDot)
+					{
+						// run a line trace between the character and the sensed actor
+						FCollisionQueryParams QueryParams;
+						QueryParams.AddIgnoredActor(LambdaInstanceData->Character);
+						QueryParams.AddIgnoredActor(SensedActor);
+
+						FHitResult OutHit;
+
+						// we have direct line of sight if this trace is unobstructed
+						bDirectLOS = !LambdaInstanceData->Character->GetWorld()->LineTraceSingleByChannel(OutHit, LambdaInstanceData->Character->GetActorLocation(), SensedActor->GetActorLocation(), ECC_Visibility, QueryParams);
+
+					}
+
+					// check if we have a direct line of sight to the stimulus
+					if (bDirectLOS)
+					{
+						// set the controller's target
+						LambdaInstanceData->Controller->SetCurrentTarget(SensedActor);
+
+						// set the task output
+						LambdaInstanceData->TargetActor = SensedActor;
+
+						// set the flags
+						LambdaInstanceData->bHasTarget = true;
+						LambdaInstanceData->bHasInvestigateLocation = false;
+
+					// no direct line of sight to target
+					} else {
+
+						// if we already have a target, ignore the partial sense and keep on them
+						if (!IsValid(LambdaInstanceData->TargetActor))
+						{
+							// is this stimulus stronger than the last one we had?
+							if (Stimulus.Strength > LambdaInstanceData->LastStimulusStrength)
+							{
+								// update the stimulus strength
+								LambdaInstanceData->LastStimulusStrength = Stimulus.Strength;
+
+								// set the investigate location
+								LambdaInstanceData->InvestigateLocation = Stimulus.StimulusLocation;
+
+								// set the investigate flag
+								LambdaInstanceData->bHasInvestigateLocation = true;
+							}
+						}
+					}
+				}
+			}
+		);
+
+		// bind the perception forgotten delegate on the controller
+		InstanceData.Controller->OnShooterPerceptionForgotten.BindLambda(
+			[WeakContext = Context.MakeWeakExecutionContext()](AActor* SensedActor)
+			{
+				// get the instance data inside the lambda
+				FInstanceDataType* LambdaInstanceData = WeakContext.MakeStrongExecutionContext().GetInstanceDataPtr<FInstanceDataType>();
+
+				if (!LambdaInstanceData)
+				{
+					return;
+				}
+
+				bool bForget = false;
+
+				// are we forgetting the current target?
+				if (SensedActor == LambdaInstanceData->TargetActor)
+				{
+					bForget = true;
+
+				} else {
+
+					// are we forgetting about a partial sense?
+					if (!IsValid(LambdaInstanceData->TargetActor))
+					{
+						bForget = true;
+					}
+				}
+
+				if (bForget)
+				{
+					// clear the target
+					LambdaInstanceData->TargetActor = nullptr;
+
+					// clear the flags
+					LambdaInstanceData->bHasInvestigateLocation = false;
+					LambdaInstanceData->bHasTarget = false;
+
+					// reset the stimulus strength
+					LambdaInstanceData->LastStimulusStrength = 0.0f;
+
+					// clear the target on the controller
+					LambdaInstanceData->Controller->ClearCurrentTarget();
+					LambdaInstanceData->Controller->ClearFocus(EAIFocusPriority::Gameplay);
+				}
+
+			}
+		);
+	}
+
+	return EStateTreeRunStatus::Running;
+}
+
+void FStateTreeSenseEnemiesTask::ExitState(FStateTreeExecutionContext& Context, const FStateTreeTransitionResult& Transition) const
+{
+	// have we transitioned to another state?
+	if (Transition.ChangeType == EStateTreeStateChangeType::Changed)
+	{
+		// get the instance data
+		FInstanceDataType& InstanceData = Context.GetInstanceData(*this);
+
+		// unbind the perception delegates
+		InstanceData.Controller->OnShooterPerceptionUpdated.Unbind();
+		InstanceData.Controller->OnShooterPerceptionForgotten.Unbind();
+	}
+}
+
+#if WITH_EDITOR
+FText FStateTreeSenseEnemiesTask::GetDescription(const FGuid& ID, FStateTreeDataView InstanceDataView, const IStateTreeBindingLookup& BindingLookup, EStateTreeNodeFormatting Formatting /*= EStateTreeNodeFormatting::Text*/) const
+{
+	return LOCTEXT("StateTreeTaskSenseEnemiesDescription", "<b>Sense Enemies</b>");
+}
+#endif // WITH_EDITOR
\ No newline at end of file
diff --git a/Templates/TP_FirstPerson/Source/TP_FirstPerson/Variant_Shooter/AI/ShooterStateTreeUtility.h b/Templates/TP_FirstPerson/Source/TP_FirstPerson/Variant_Shooter/AI/ShooterStateTreeUtility.h
new file mode 100644
--- /dev/null
+++ b/Templates/TP_FirstPerson/Source/TP_FirstPerson/Variant_Shooter/AI/ShooterStateTreeUtility.h
@@ -0,0 +1,309 @@
+// Copyright Epic Games, Inc. All Rights Reserved.
+
+#pragma once
+
+#include "CoreMinimal.h"
+#include "StateTreeTaskBase.h"
+#include "StateTreeConditionBase.h"
+
+#include "ShooterStateTreeUtility.generated.h"
+
+class AShooterNPC;
+class AAIController;
+class AShooterAIController;
+
+/**
+ *  Instance data struct for the FStateTreeLineOfSightToTargetCondition condition
+ */
+USTRUCT()
+struct FStateTreeLineOfSightToTargetConditionInstanceData
+{
+	GENERATED_BODY()
+	
+	/** Targeting character */
+	UPROPERTY(EditAnywhere, Category = "Context")
+	AShooterNPC* Character;
+
+	/** Target to check line of sight for */
+	UPROPERTY(EditAnywhere, Category = "Condition")
+	AActor* Target;
+
+	/** Max allowed line of sight cone angle, in degrees */
+	UPROPERTY(EditAnywhere, Category = "Condition")
+	float LineOfSightConeAngle = 35.0f;
+
+	/** Number of vertical line of sight checks to run to try and get around low obstacles */
+	UPROPERTY(EditAnywhere, Category = "Condition")
+	int32 NumberOfVerticalLineOfSightChecks = 5;
+
+	/** If true, the condition passes if the character has line of sight */
+	UPROPERTY(EditAnywhere, Category = "Condition")
+	bool bMustHaveLineOfSight = true;
+};
+STATETREE_POD_INSTANCEDATA(FStateTreeLineOfSightToTargetConditionInstanceData);
+
+/**
+ *  StateTree condition to check if the character is grounded
+ */
+USTRUCT(DisplayName = "Has Line of Sight to Target", Category="Shooter")
+struct FStateTreeLineOfSightToTargetCondition : public FStateTreeConditionCommonBase
+{
+	GENERATED_BODY()
+
+	/** Set the instance data type */
+	using FInstanceDataType = FStateTreeLineOfSightToTargetConditionInstanceData;
+	virtual const UStruct* GetInstanceDataType() const override { return FInstanceDataType::StaticStruct(); }
+
+	/** Default constructor */
+	FStateTreeLineOfSightToTargetCondition() = default;
+	
+	/** Tests the StateTree condition */
+	virtual bool TestCondition(FStateTreeExecutionContext& Context) const override;
+
+#if WITH_EDITOR
+	/** Provides the description string */
+	virtual FText GetDescription(const FGuid& ID, FStateTreeDataView InstanceDataView, const IStateTreeBindingLookup& BindingLookup, EStateTreeNodeFormatting Formatting = EStateTreeNodeFormatting::Text) const override;
+#endif
+
+};
+
+////////////////////////////////////////////////////////////////////
+
+/**
+ *  Instance data struct for the Face Towards Actor StateTree task
+ */
+USTRUCT()
+struct FStateTreeFaceActorInstanceData
+{
+	GENERATED_BODY()
+
+	/** AI Controller that will determine the focused actor */
+	UPROPERTY(EditAnywhere, Category = Context)
+	TObjectPtr<AAIController> Controller;
+
+	/** Actor that will be faced towards */
+	UPROPERTY(EditAnywhere, Category = Input)
+	TObjectPtr<AActor> ActorToFaceTowards;
+};
+
+/**
+ *  StateTree task to face an AI-Controlled Pawn towards an Actor
+ */
+USTRUCT(meta=(DisplayName="Face Towards Actor", Category="Shooter"))
+struct FStateTreeFaceActorTask : public FStateTreeTaskCommonBase
+{
+	GENERATED_BODY()
+
+	/* Ensure we're using the correct instance data struct */
+	using FInstanceDataType = FStateTreeFaceActorInstanceData;
+	virtual const UStruct* GetInstanceDataType() const override { return FInstanceDataType::StaticStruct(); }
+
+	/** Runs when the owning state is entered */
+	virtual EStateTreeRunStatus EnterState(FStateTreeExecutionContext& Context, const FStateTreeTransitionResult& Transition) const override;
+
+	/** Runs when the owning state is ended */
+	virtual void ExitState(FStateTreeExecutionContext& Context, const FStateTreeTransitionResult& Transition) const override;
+
+#if WITH_EDITOR
+	virtual FText GetDescription(const FGuid& ID, FStateTreeDataView InstanceDataView, const IStateTreeBindingLookup& BindingLookup, EStateTreeNodeFormatting Formatting = EStateTreeNodeFormatting::Text) const override;
+#endif // WITH_EDITOR
+};
+
+////////////////////////////////////////////////////////////////////
+
+/**
+ *  Instance data struct for the Face Towards Location StateTree task
+ */
+USTRUCT()
+struct FStateTreeFaceLocationInstanceData
+{
+	GENERATED_BODY()
+
+	/** AI Controller that will determine the focused location */
+	UPROPERTY(EditAnywhere, Category = Context)
+	TObjectPtr<AAIController> Controller;
+
+	/** Location that will be faced towards */
+	UPROPERTY(EditAnywhere, Category = Parameter)
+	FVector FaceLocation;
+};
+
+/**
+ *  StateTree task to face an AI-Controlled Pawn towards a world location
+ */
+USTRUCT(meta=(DisplayName="Face Towards Location", Category="Shooter"))
+struct FStateTreeFaceLocationTask : public FStateTreeTaskCommonBase
+{
+	GENERATED_BODY()
+
+	/* Ensure we're using the correct instance data struct */
+	using FInstanceDataType = FStateTreeFaceLocationInstanceData;
+	virtual const UStruct* GetInstanceDataType() const override { return FInstanceDataType::StaticStruct(); }
+
+	/** Runs when the owning state is entered */
+	virtual EStateTreeRunStatus EnterState(FStateTreeExecutionContext& Context, const FStateTreeTransitionResult& Transition) const override;
+
+	/** Runs when the owning state is ended */
+	virtual void ExitState(FStateTreeExecutionContext& Context, const FStateTreeTransitionResult& Transition) const override;
+
+#if WITH_EDITOR
+	virtual FText GetDescription(const FGuid& ID, FStateTreeDataView InstanceDataView, const IStateTreeBindingLookup& BindingLookup, EStateTreeNodeFormatting Formatting = EStateTreeNodeFormatting::Text) const override;
+#endif // WITH_EDITOR
+};
+
+////////////////////////////////////////////////////////////////////
+
+/**
+ *  Instance data struct for the Set Random Float StateTree task
+ */
+USTRUCT()
+struct FStateTreeSetRandomFloatData
+{
+	GENERATED_BODY()
+
+	/** Minimum random value */
+	UPROPERTY(EditAnywhere, Category = Parameter)
+	float MinValue = 0.0f;
+
+	/** Maximum random value */
+	UPROPERTY(EditAnywhere, Category = Parameter)
+	float MaxValue = 0.0f;
+
+	/** Output calculated value */
+	UPROPERTY(EditAnywhere, Category = Output)
+	float OutValue = 0.0f;
+};
+
+/**
+ *  StateTree task to calculate a random float value within the specified range
+ */
+USTRUCT(meta=(DisplayName="Set Random Float", Category="Shooter"))
+struct FStateTreeSetRandomFloatTask : public FStateTreeTaskCommonBase
+{
+	GENERATED_BODY()
+
+	/* Ensure we're using the correct instance data struct */
+	using FInstanceDataType = FStateTreeSetRandomFloatData;
+	virtual const UStruct* GetInstanceDataType() const override { return FInstanceDataType::StaticStruct(); }
+
+	/** Runs when the owning state is entered */
+	virtual EStateTreeRunStatus EnterState(FStateTreeExecutionContext& Context, const FStateTreeTransitionResult& Transition) const override;
+
+#if WITH_EDITOR
+	virtual FText GetDescription(const FGuid& ID, FStateTreeDataView InstanceDataView, const IStateTreeBindingLookup& BindingLookup, EStateTreeNodeFormatting Formatting = EStateTreeNodeFormatting::Text) const override;
+#endif // WITH_EDITOR
+};
+
+////////////////////////////////////////////////////////////////////
+
+/**
+ *  Instance data struct for the Shoot At Target StateTree task
+ */
+USTRUCT()
+struct FStateTreeShootAtTargetInstanceData
+{
+	GENERATED_BODY()
+
+	/** NPC that will do the shooting */
+	UPROPERTY(EditAnywhere, Category = Context)
+	TObjectPtr<AShooterNPC> Character;
+
+	/** Target to shoot at */
+	UPROPERTY(EditAnywhere, Category = Input)
+	TObjectPtr<AActor> Target;
+};
+
+/**
+ *  StateTree task to have an NPC shoot at an actor
+ */
+USTRUCT(meta=(DisplayName="Shoot at Target", Category="Shooter"))
+struct FStateTreeShootAtTargetTask : public FStateTreeTaskCommonBase
+{
+	GENERATED_BODY()
+
+	/* Ensure we're using the correct instance data struct */
+	using FInstanceDataType = FStateTreeShootAtTargetInstanceData;
+	virtual const UStruct* GetInstanceDataType() const override { return FInstanceDataType::StaticStruct(); }
+
+	/** Runs when the owning state is entered */
+	virtual EStateTreeRunStatus EnterState(FStateTreeExecutionContext& Context, const FStateTreeTransitionResult& Transition) const override;
+
+	/** Runs when the owning state is ended */
+	virtual void ExitState(FStateTreeExecutionContext& Context, const FStateTreeTransitionResult& Transition) const override;
+
+#if WITH_EDITOR
+	virtual FText GetDescription(const FGuid& ID, FStateTreeDataView InstanceDataView, const IStateTreeBindingLookup& BindingLookup, EStateTreeNodeFormatting Formatting = EStateTreeNodeFormatting::Text) const override;
+#endif // WITH_EDITOR
+};
+
+////////////////////////////////////////////////////////////////////
+
+/**
+ *  Instance data struct for the Sense Enemies StateTree task
+ */
+USTRUCT()
+struct FStateTreeSenseEnemiesInstanceData
+{
+	GENERATED_BODY()
+
+	/** Sensing AI Controller */
+	UPROPERTY(EditAnywhere, Category = Context)
+	TObjectPtr<AShooterAIController> Controller;
+
+	/** Sensing NPC */
+	UPROPERTY(EditAnywhere, Category = Context)
+	TObjectPtr<AShooterNPC> Character;
+
+	/** Sensed actor to target */
+	UPROPERTY(EditAnywhere, Category = Output)
+	TObjectPtr<AActor> TargetActor;
+
+	/** Sensed location to investigate */
+	UPROPERTY(EditAnywhere, Category = Output)
+	FVector InvestigateLocation;
+
+	/** True if a target was successfully sensed */
+	UPROPERTY(EditAnywhere, Category = Output)
+	bool bHasTarget = false;
+
+	/** True if an investigate location was successfully sensed */
+	UPROPERTY(EditAnywhere, Category = Output)
+	bool bHasInvestigateLocation = false;
+
+	/** Tag required on sensed actors */
+	UPROPERTY(EditAnywhere, Category = Parameter)
+	FName SenseTag = FName("Player");
+
+	/** Line of sight cone half angle to consider a full sense */
+	UPROPERTY(EditAnywhere, Category = Parameter)
+	float DirectLineOfSightCone = 85.0f;
+
+	/** Strength of the last processed stimulus */
+	UPROPERTY(EditAnywhere)
+	float LastStimulusStrength = 0.0f;
+};
+
+/**
+ *  StateTree task to have an NPC process AI Perceptions and sense nearby enemies
+ */
+USTRUCT(meta=(DisplayName="Sense Enemies", Category="Shooter"))
+struct FStateTreeSenseEnemiesTask : public FStateTreeTaskCommonBase
+{
+	GENERATED_BODY()
+
+	/* Ensure we're using the correct instance data struct */
+	using FInstanceDataType = FStateTreeSenseEnemiesInstanceData;
+	virtual const UStruct* GetInstanceDataType() const override { return FInstanceDataType::StaticStruct(); }
+
+	/** Runs when the owning state is entered */
+	virtual EStateTreeRunStatus EnterState(FStateTreeExecutionContext& Context, const FStateTreeTransitionResult& Transition) const override;
+
+	/** Runs when the owning state is ended */
+	virtual void ExitState(FStateTreeExecutionContext& Context, const FStateTreeTransitionResult& Transition) const override;
+
+#if WITH_EDITOR
+	virtual FText GetDescription(const FGuid& ID, FStateTreeDataView InstanceDataView, const IStateTreeBindingLookup& BindingLookup, EStateTreeNodeFormatting Formatting = EStateTreeNodeFormatting::Text) const override;
+#endif // WITH_EDITOR
+};
+
+////////////////////////////////////////////////////////////////////
\ No newline at end of file
diff --git a/Templates/TP_FirstPerson/Source/TP_FirstPerson/Variant_Shooter/ShooterBulletCounterUI.cpp b/Templates/TP_FirstPerson/Source/TP_FirstPerson/Variant_Shooter/ShooterBulletCounterUI.cpp
new file mode 100644
--- /dev/null
+++ b/Templates/TP_FirstPerson/Source/TP_FirstPerson/Variant_Shooter/ShooterBulletCounterUI.cpp
@@ -0,0 +1,5 @@
+// Copyright Epic Games, Inc. All Rights Reserved.
+
+
+#include "Variant_Shooter/ShooterBulletCounterUI.h"
+
diff --git a/Templates/TP_FirstPerson/Source/TP_FirstPerson/Variant_Shooter/ShooterBulletCounterUI.h b/Templates/TP_FirstPerson/Source/TP_FirstPerson/Variant_Shooter/ShooterBulletCounterUI.h
new file mode 100644
--- /dev/null
+++ b/Templates/TP_FirstPerson/Source/TP_FirstPerson/Variant_Shooter/ShooterBulletCounterUI.h
@@ -0,0 +1,22 @@
+// Copyright Epic Games, Inc. All Rights Reserved.
+
+#pragma once
+
+#include "CoreMinimal.h"
+#include "Blueprint/UserWidget.h"
+#include "ShooterBulletCounterUI.generated.h"
+
+/**
+ *  Simple bullet counter UI widget for a first person shooter game
+ */
+UCLASS(abstract)
+class TP_FIRSTPERSON_API UShooterBulletCounterUI : public UUserWidget
+{
+	GENERATED_BODY()
+	
+public:
+
+	/** Allows Blueprint to update sub-widgets with the new bullet count */
+	UFUNCTION(BlueprintImplementableEvent, Category="Shooter", meta=(DisplayName = "UpdateBulletCounter"))
+	void BP_UpdateBulletCounter(int32 MagazineSize, int32 BulletCount);
+};
diff --git a/Templates/TP_FirstPerson/Source/TP_FirstPerson/Variant_Shooter/ShooterCharacter.cpp b/Templates/TP_FirstPerson/Source/TP_FirstPerson/Variant_Shooter/ShooterCharacter.cpp
new file mode 100644
--- /dev/null
+++ b/Templates/TP_FirstPerson/Source/TP_FirstPerson/Variant_Shooter/ShooterCharacter.cpp
@@ -0,0 +1,230 @@
+// Copyright Epic Games, Inc. All Rights Reserved.
+
+
+#include "ShooterCharacter.h"
+#include "ShooterWeapon.h"
+#include "EnhancedInputComponent.h"
+#include "Components/InputComponent.h"
+#include "Components/PawnNoiseEmitterComponent.h"
+#include "GameFramework/CharacterMovementComponent.h"
+#include "Components/SkeletalMeshComponent.h"
+#include "Engine/World.h"
+#include "Camera/CameraComponent.h"
+
+AShooterCharacter::AShooterCharacter()
+{
+	// create the noise emitter component
+	PawnNoiseEmitter = CreateDefaultSubobject<UPawnNoiseEmitterComponent>(TEXT("Pawn Noise Emitter"));
+
+	// configure movement
+	GetCharacterMovement()->RotationRate = FRotator(0.0f, 600.0f, 0.0f);
+}
+
+void AShooterCharacter::SetupPlayerInputComponent(UInputComponent* PlayerInputComponent)
+{
+	// base class handles move, aim and jump inputs
+	Super::SetupPlayerInputComponent(PlayerInputComponent);
+
+	// Set up action bindings
+	if (UEnhancedInputComponent* EnhancedInputComponent = Cast<UEnhancedInputComponent>(PlayerInputComponent))
+	{
+		// Firing
+		EnhancedInputComponent->BindAction(FireAction, ETriggerEvent::Started, this, &AShooterCharacter::DoStartFiring);
+		EnhancedInputComponent->BindAction(FireAction, ETriggerEvent::Completed, this, &AShooterCharacter::DoStopFiring);
+
+		// Switch weapon
+		EnhancedInputComponent->BindAction(SwitchWeaponAction, ETriggerEvent::Triggered, this, &AShooterCharacter::DoSwitchWeapon);
+	}
+
+}
+
+float AShooterCharacter::TakeDamage(float Damage, struct FDamageEvent const& DamageEvent, AController* EventInstigator, AActor* DamageCauser)
+{
+	// ignore if already dead
+	if (CurrentHP <= 0.0f)
+	{
+		return 0.0f;
+	}
+
+	// Reduce HP
+	CurrentHP -= Damage;
+
+	// Have we depleted HP?
+	if (CurrentHP <= 0.0f)
+	{
+		// deactivate the weapon
+		CurrentWeapon->DeactivateWeapon();
+
+		// reset the bullet counter UI
+		OnBulletCountUpdated.Broadcast(0, 0);
+
+		// destroy this character
+		Destroy();
+	}
+
+	return Damage;
+}
+
+void AShooterCharacter::DoStartFiring()
+{
+	// fire the current weapon
+	if (CurrentWeapon)
+	{
+		CurrentWeapon->StartFiring();
+	}
+}
+
+void AShooterCharacter::DoStopFiring()
+{
+	// stop firing the current weapon
+	if (CurrentWeapon)
+	{
+		CurrentWeapon->StopFiring();
+	}
+}
+
+void AShooterCharacter::DoSwitchWeapon()
+{
+	// ensure we have at least two weapons two switch between
+	if (OwnedWeapons.Num() > 1)
+	{
+		// deactivate the old weapon
+		CurrentWeapon->DeactivateWeapon();
+
+		// find the index of the current weapon in the owned list
+		int32 WeaponIndex = OwnedWeapons.Find(CurrentWeapon);
+
+		// is this the last weapon?
+		if (WeaponIndex == OwnedWeapons.Num() - 1)
+		{
+			// loop back to the beginning of the array
+			WeaponIndex = 0;
+		}
+		else {
+			// select the next weapon index
+			++WeaponIndex;
+		}
+
+		// set the new weapon as current
+		CurrentWeapon = OwnedWeapons[WeaponIndex];
+
+		// activate the new weapon
+		CurrentWeapon->ActivateWeapon();
+	}
+}
+
+void AShooterCharacter::AttachWeaponMeshes(AShooterWeapon* Weapon)
+{
+	const FAttachmentTransformRules AttachmentRule(EAttachmentRule::SnapToTarget, false);
+
+	// attach the weapon actor
+	Weapon->AttachToActor(this, AttachmentRule);
+
+	// attach the weapon meshes
+	Weapon->GetFirstPersonMesh()->AttachToComponent(GetFirstPersonMesh(), AttachmentRule, FirstPersonWeaponSocket);
+	Weapon->GetThirdPersonMesh()->AttachToComponent(GetMesh(), AttachmentRule, FirstPersonWeaponSocket);
+	
+}
+
+void AShooterCharacter::PlayFiringMontage(UAnimMontage* Montage)
+{
+	
+}
+
+void AShooterCharacter::AddWeaponRecoil(float Recoil)
+{
+	// apply the recoil as pitch input
+	AddControllerPitchInput(Recoil);
+}
+
+void AShooterCharacter::UpdateWeaponHUD(int32 CurrentAmmo, int32 MagazineSize)
+{
+	OnBulletCountUpdated.Broadcast(MagazineSize, CurrentAmmo);
+}
+
+FVector AShooterCharacter::GetWeaponTargetLocation()
+{
+	// trace ahead from the camera viewpoint
+	FHitResult OutHit;
+
+	const FVector Start = GetFirstPersonCameraComponent()->GetComponentLocation();
+	const FVector End = Start + (GetFirstPersonCameraComponent()->GetForwardVector() * MaxAimDistance);
+
+	FCollisionQueryParams QueryParams;
+	QueryParams.AddIgnoredActor(this);
+
+	GetWorld()->LineTraceSingleByChannel(OutHit, Start, End, ECC_Visibility, QueryParams);
+
+	// return either the impact point or the trace end
+	return OutHit.bBlockingHit ? OutHit.ImpactPoint : OutHit.TraceEnd;
+}
+
+void AShooterCharacter::AddWeaponClass(const TSubclassOf<AShooterWeapon>& WeaponClass)
+{
+	// do we already own this weapon?
+	AShooterWeapon* OwnedWeapon = FindWeaponOfType(WeaponClass);
+
+	if (!OwnedWeapon)
+	{
+		// spawn the new weapon
+		FActorSpawnParameters SpawnParams;
+		SpawnParams.Owner = this;
+		SpawnParams.Instigator = this;
+		SpawnParams.SpawnCollisionHandlingOverride = ESpawnActorCollisionHandlingMethod::AlwaysSpawn;
+		SpawnParams.TransformScaleMethod = ESpawnActorScaleMethod::MultiplyWithRoot;
+
+		AShooterWeapon* AddedWeapon = GetWorld()->SpawnActor<AShooterWeapon>(WeaponClass, GetActorTransform(), SpawnParams);
+
+		if (AddedWeapon)
+		{
+			// add the weapon to the owned list
+			OwnedWeapons.Add(AddedWeapon);
+
+			// if we have an existing weapon, deactivate it
+			if (CurrentWeapon)
+			{
+				CurrentWeapon->DeactivateWeapon();
+			}
+
+			// switch to the new weapon
+			CurrentWeapon = AddedWeapon;
+			CurrentWeapon->ActivateWeapon();
+		}
+	}
+}
+
+void AShooterCharacter::OnWeaponActivated(AShooterWeapon* Weapon)
+{
+	// update the bullet counter
+	OnBulletCountUpdated.Broadcast(Weapon->GetMagazineSize(), Weapon->GetBulletCount());
+
+	// set the character mesh AnimInstances
+	GetFirstPersonMesh()->SetAnimInstanceClass(Weapon->GetFirstPersonAnimInstanceClass());
+	GetMesh()->SetAnimInstanceClass(Weapon->GetThirdPersonAnimInstanceClass());
+}
+
+void AShooterCharacter::OnWeaponDeactivated(AShooterWeapon* Weapon)
+{
+	// unused
+}
+
+void AShooterCharacter::OnSemiWeaponRefire()
+{
+	// unused
+}
+
+AShooterWeapon* AShooterCharacter::FindWeaponOfType(TSubclassOf<AShooterWeapon> WeaponClass) const
+{
+	// check each owned weapon
+	for (AShooterWeapon* Weapon : OwnedWeapons)
+	{
+		if (Weapon->IsA(WeaponClass))
+		{
+			return Weapon;
+		}
+	}
+
+	// weapon not found
+	return nullptr;
+
+}
diff --git a/Templates/TP_FirstPerson/Source/TP_FirstPerson/Variant_Shooter/ShooterCharacter.h b/Templates/TP_FirstPerson/Source/TP_FirstPerson/Variant_Shooter/ShooterCharacter.h
new file mode 100644
--- /dev/null
+++ b/Templates/TP_FirstPerson/Source/TP_FirstPerson/Variant_Shooter/ShooterCharacter.h
@@ -0,0 +1,138 @@
+// Copyright Epic Games, Inc. All Rights Reserved.
+
+#pragma once
+
+#include "CoreMinimal.h"
+#include "TP_FirstPersonCharacter.h"
+#include "ShooterWeaponHolder.h"
+#include "ShooterCharacter.generated.h"
+
+class AShooterWeapon;
+class UInputAction;
+class UInputComponent;
+class UPawnNoiseEmitterComponent;
+
+DECLARE_DYNAMIC_MULTICAST_DELEGATE_TwoParams(FBulletCountUpdatedDelegate, int32, MagazineSize, int32, Bullets);
+
+/**
+ *  A player controllable first person shooter character
+ *  Manages a weapon inventory through the IShooterWeaponHolder interface
+ *  Manages health and death
+ */
+UCLASS(abstract)
+class TP_FIRSTPERSON_API AShooterCharacter : public ATP_FirstPersonCharacter, public IShooterWeaponHolder
+{
+	GENERATED_BODY()
+	
+	/** AI Noise emitter component */
+	UPROPERTY(VisibleAnywhere, BlueprintReadOnly, Category="Components", meta = (AllowPrivateAccess = "true"))
+	UPawnNoiseEmitterComponent* PawnNoiseEmitter;
+
+protected:
+
+	/** Fire weapon input action */
+	UPROPERTY(EditAnywhere, BlueprintReadOnly, Category ="Input")
+	UInputAction* FireAction;
+
+	/** Switch weapon input action */
+	UPROPERTY(EditAnywhere, BlueprintReadOnly, Category ="Input")
+	UInputAction* SwitchWeaponAction;
+
+	/** Name of the first person mesh weapon socket */
+	UPROPERTY(EditAnywhere, BlueprintReadOnly, Category ="Weapons")
+	FName FirstPersonWeaponSocket = FName("HandGrip_R");
+
+	/** Name of the third person mesh weapon socket */
+	UPROPERTY(EditAnywhere, BlueprintReadOnly, Category ="Weapons")
+	FName ThirdPersonWeaponSocket = FName("HandGrip_R");
+
+	/** Max distance to use for aim traces */
+	UPROPERTY(EditAnywhere, BlueprintReadOnly, Category ="Aim")
+	float MaxAimDistance = 10000.0f;
+
+	/** Current HP remaining to this character */
+	UPROPERTY(EditAnywhere, Category="Health")
+	float CurrentHP = 500.0f;
+
+	/** Team ID for this character*/
+	UPROPERTY(EditAnywhere, Category="Team")
+	uint8 TeamByte = 0;
+
+	/** List of weapons picked up by the character */
+	TArray<AShooterWeapon*> OwnedWeapons;
+
+	/** Weapon currently equipped and ready to shoot with */
+	TObjectPtr<AShooterWeapon> CurrentWeapon;
+
+public:
+
+	/** Bullet count updated delegate */
+	FBulletCountUpdatedDelegate OnBulletCountUpdated;
+
+public:
+
+	/** Constructor */
+	AShooterCharacter();
+
+protected:
+
+	/** Set up input action bindings */
+	virtual void SetupPlayerInputComponent(UInputComponent* InputComponent) override;
+
+public:
+
+	/** Handle incoming damage */
+	virtual float TakeDamage(float Damage, struct FDamageEvent const& DamageEvent, AController* EventInstigator, AActor* DamageCauser) override;
+
+public:
+
+	/** Handles start firing input */
+	UFUNCTION(BlueprintCallable, Category="Input")
+	void DoStartFiring();
+
+	/** Handles stop firing input */
+	UFUNCTION(BlueprintCallable, Category="Input")
+	void DoStopFiring();
+
+	/** Handles switch weapon input */
+	UFUNCTION(BlueprintCallable, Category="Input")
+	void DoSwitchWeapon();
+
+public:
+
+	//~Begin IShooterWeaponHolder interface
+
+	/** Attaches a weapon's meshes to the owner */
+	virtual void AttachWeaponMeshes(AShooterWeapon* Weapon) override;
+
+	/** Plays the firing montage for the weapon */
+	virtual void PlayFiringMontage(UAnimMontage* Montage) override;
+
+	/** Applies weapon recoil to the owner */
+	virtual void AddWeaponRecoil(float Recoil) override;
+
+	/** Updates the weapon's HUD with the current ammo count */
+	virtual void UpdateWeaponHUD(int32 CurrentAmmo, int32 MagazineSize) override;
+
+	/** Calculates and returns the aim location for the weapon */
+	virtual FVector GetWeaponTargetLocation() override;
+
+	/** Gives a weapon of this class to the owner */
+	virtual void AddWeaponClass(const TSubclassOf<AShooterWeapon>& WeaponClass) override;
+
+	/** Activates the passed weapon */
+	virtual void OnWeaponActivated(AShooterWeapon* Weapon) override;
+
+	/** Deactivates the passed weapon */
+	virtual void OnWeaponDeactivated(AShooterWeapon* Weapon) override;
+
+	/** Notifies the owner that the weapon cooldown has expired and it's ready to shoot again */
+	virtual void OnSemiWeaponRefire() override;
+
+	//~End IShooterWeaponHolder interface
+
+protected:
+
+	/** Returns true if the character already owns a weapon of the given class */
+	AShooterWeapon* FindWeaponOfType(TSubclassOf<AShooterWeapon> WeaponClass) const;
+};
diff --git a/Templates/TP_FirstPerson/Source/TP_FirstPerson/Variant_Shooter/ShooterGameMode.cpp b/Templates/TP_FirstPerson/Source/TP_FirstPerson/Variant_Shooter/ShooterGameMode.cpp
new file mode 100644
--- /dev/null
+++ b/Templates/TP_FirstPerson/Source/TP_FirstPerson/Variant_Shooter/ShooterGameMode.cpp
@@ -0,0 +1,32 @@
+// Copyright Epic Games, Inc. All Rights Reserved.
+
+
+#include "Variant_Shooter/ShooterGameMode.h"
+#include "ShooterUI.h"
+#include "Kismet/GameplayStatics.h"
+#include "Engine/World.h"
+
+void AShooterGameMode::BeginPlay()
+{
+	Super::BeginPlay();
+
+	// create the UI
+	ShooterUI = CreateWidget<UShooterUI>(UGameplayStatics::GetPlayerController(GetWorld(), 0), ShooterUIClass);
+	ShooterUI->AddToViewport(0);
+}
+
+void AShooterGameMode::IncrementTeamScore(uint8 TeamByte)
+{
+	// retrieve the team score if any
+	int32 Score = 0;
+	if (int32* FoundScore = TeamScores.Find(TeamByte))
+	{
+		Score = *FoundScore;
+	}
+
+	// increment the score for the given team
+	TeamScores.Add(TeamByte, Score + 1);
+
+	// update the UI
+	ShooterUI->BP_UpdateScore(TeamByte, Score);
+}
diff --git a/Templates/TP_FirstPerson/Source/TP_FirstPerson/Variant_Shooter/ShooterGameMode.h b/Templates/TP_FirstPerson/Source/TP_FirstPerson/Variant_Shooter/ShooterGameMode.h
new file mode 100644
--- /dev/null
+++ b/Templates/TP_FirstPerson/Source/TP_FirstPerson/Variant_Shooter/ShooterGameMode.h
@@ -0,0 +1,42 @@
+// Copyright Epic Games, Inc. All Rights Reserved.
+
+#pragma once
+
+#include "CoreMinimal.h"
+#include "GameFramework/GameModeBase.h"
+#include "ShooterGameMode.generated.h"
+
+class UShooterUI;
+
+/**
+ *  Simple GameMode for a first person shooter game
+ *  Manages game UI
+ *  Keeps track of team scores
+ */
+UCLASS(abstract)
+class TP_FIRSTPERSON_API AShooterGameMode : public AGameModeBase
+{
+	GENERATED_BODY()
+	
+protected:
+
+	/** Type of UI widget to spawn */
+	UPROPERTY(EditAnywhere, Category="Shooter")
+	TSubclassOf<UShooterUI> ShooterUIClass;
+
+	/** Pointer to the UI widget */
+	TObjectPtr<UShooterUI> ShooterUI;
+
+	/** Map of scores by team ID */
+	TMap<uint8, int32> TeamScores;
+
+protected:
+
+	/** Gameplay initialization */
+	virtual void BeginPlay() override;
+
+public:
+
+	/** Increases the score for the given team */
+	void IncrementTeamScore(uint8 TeamByte);
+};
diff --git a/Templates/TP_FirstPerson/Source/TP_FirstPerson/Variant_Shooter/ShooterPickup.cpp b/Templates/TP_FirstPerson/Source/TP_FirstPerson/Variant_Shooter/ShooterPickup.cpp
new file mode 100644
--- /dev/null
+++ b/Templates/TP_FirstPerson/Source/TP_FirstPerson/Variant_Shooter/ShooterPickup.cpp
@@ -0,0 +1,108 @@
+// Copyright Epic Games, Inc. All Rights Reserved.
+
+
+#include "Variant_Shooter/ShooterPickup.h"
+#include "Components/SceneComponent.h"
+#include "Components/SphereComponent.h"
+#include "Components/StaticMeshComponent.h"
+#include "ShooterWeaponHolder.h"
+#include "ShooterWeapon.h"
+#include "Engine/World.h"
+#include "TimerManager.h"
+
+AShooterPickup::AShooterPickup()
+{
+ 	PrimaryActorTick.bCanEverTick = true;
+
+	// create the root
+	RootComponent = CreateDefaultSubobject<USceneComponent>(TEXT("Root"));
+
+	// create the collision sphere
+	SphereCollision = CreateDefaultSubobject<USphereComponent>(TEXT("Sphere Collision"));
+	SphereCollision->SetupAttachment(RootComponent);
+
+	SphereCollision->SetRelativeLocation(FVector(0.0f, 0.0f, 84.0f));
+	SphereCollision->SetCollisionEnabled(ECollisionEnabled::QueryOnly);
+	SphereCollision->SetCollisionObjectType(ECC_WorldStatic);
+	SphereCollision->SetCollisionResponseToAllChannels(ECR_Ignore);
+	SphereCollision->SetCollisionResponseToChannel(ECC_Pawn, ECR_Overlap);
+	SphereCollision->bFillCollisionUnderneathForNavmesh = true;
+
+	// subscribe to the collision overlap on the sphere
+	SphereCollision->OnComponentBeginOverlap.AddDynamic(this, &AShooterPickup::OnOverlap);
+
+	// create the mesh
+	Mesh = CreateDefaultSubobject<UStaticMeshComponent>(TEXT("Mesh"));
+	Mesh->SetupAttachment(SphereCollision);
+
+	Mesh->SetCollisionProfileName(FName("NoCollision"));
+}
+
+void AShooterPickup::OnConstruction(const FTransform& Transform)
+{
+	Super::OnConstruction(Transform);
+
+	if (FWeaponTableRow* WeaponData = WeaponType.GetRow<FWeaponTableRow>(FString()))
+	{
+		// set the mesh
+		Mesh->SetStaticMesh(WeaponData->StaticMesh.LoadSynchronous());
+	}
+}
+
+void AShooterPickup::BeginPlay()
+{
+	Super::BeginPlay();
+
+	if (FWeaponTableRow* WeaponData = WeaponType.GetRow<FWeaponTableRow>(FString()))
+	{
+		// copy the weapon class
+		WeaponClass = WeaponData->WeaponClass.LoadSynchronous();
+	}
+}
+
+void AShooterPickup::EndPlay(const EEndPlayReason::Type EndPlayReason)
+{
+	Super::EndPlay(EndPlayReason);
+
+	// clear the respawn timer
+	GetWorld()->GetTimerManager().ClearTimer(RespawnTimer);
+}
+
+void AShooterPickup::OnOverlap(UPrimitiveComponent* OverlappedComponent, AActor* OtherActor, UPrimitiveComponent* OtherComp, int32 OtherBodyIndex, bool bFromSweep, const FHitResult& SweepResult)
+{
+	// have we collided against a weapon holder?
+	if (IShooterWeaponHolder* WeaponHolder = Cast<IShooterWeaponHolder>(OtherActor))
+	{
+		WeaponHolder->AddWeaponClass(WeaponClass);
+
+		// hide this mesh
+		SetActorHiddenInGame(true);
+
+		// disable collision
+		SetActorEnableCollision(false);
+
+		// disable ticking
+		SetActorTickEnabled(false);
+
+		// schedule the respawn
+		GetWorld()->GetTimerManager().SetTimer(RespawnTimer, this, &AShooterPickup::RespawnPickup, RespawnTime, false);
+	}
+}
+
+void AShooterPickup::RespawnPickup()
+{
+	// unhide this pickup
+	SetActorHiddenInGame(false);
+
+	// call the BP handler
+	BP_OnRespawn();
+}
+
+void AShooterPickup::FinishRespawn()
+{
+	// enable collision
+	SetActorEnableCollision(true);
+
+	// enable tick
+	SetActorTickEnabled(true);
+}
diff --git a/Templates/TP_FirstPerson/Source/TP_FirstPerson/Variant_Shooter/ShooterPickup.h b/Templates/TP_FirstPerson/Source/TP_FirstPerson/Variant_Shooter/ShooterPickup.h
new file mode 100644
--- /dev/null
+++ b/Templates/TP_FirstPerson/Source/TP_FirstPerson/Variant_Shooter/ShooterPickup.h
@@ -0,0 +1,96 @@
+// Copyright Epic Games, Inc. All Rights Reserved.
+
+#pragma once
+
+#include "CoreMinimal.h"
+#include "GameFramework/Actor.h"
+#include "Engine/DataTable.h"
+#include "Engine/StaticMesh.h"
+#include "ShooterPickup.generated.h"
+
+class USphereComponent;
+class UPrimitiveComponent;
+class AShooterWeapon;
+
+/**
+ *  Holds information about a type of weapon pickup
+ */
+USTRUCT(BlueprintType)
+struct FWeaponTableRow : public FTableRowBase
+{
+	GENERATED_BODY()
+
+	/** Mesh to display on the pickup */
+	UPROPERTY(EditAnywhere)
+	TSoftObjectPtr<UStaticMesh> StaticMesh;
+
+	/** Weapon class to grant on pickup */
+	UPROPERTY(EditAnywhere)
+	TSoftClassPtr<AShooterWeapon> WeaponClass;
+};
+
+/**
+ *  Simple shooter game weapon pickup
+ */
+UCLASS(abstract)
+class TP_FIRSTPERSON_API AShooterPickup : public AActor
+{
+	GENERATED_BODY()
+
+	/** Collision sphere */
+	UPROPERTY(VisibleAnywhere, BlueprintReadOnly, Category="Components", meta = (AllowPrivateAccess = "true"))
+	USphereComponent* SphereCollision;
+
+	/** Weapon pickup mesh. Its mesh asset is set from the weapon data table */
+	UPROPERTY(VisibleAnywhere, BlueprintReadOnly, Category="Components", meta = (AllowPrivateAccess = "true"))
+	UStaticMeshComponent* Mesh;
+	
+protected:
+
+	/** Data on the type of picked weapon and visuals of this pickup */
+	UPROPERTY(EditAnywhere, Category="Pickup")
+	FDataTableRowHandle WeaponType;
+
+	/** Type to weapon to grant on pickup. Set from the weapon data table. */
+	TSubclassOf<AShooterWeapon> WeaponClass;
+	
+	/** Time to wait before respawning this pickup */
+	UPROPERTY(EditAnywhere, Category="Pickup")
+	float RespawnTime = 4.0f;
+
+	/** Timer to respawn the pickup */
+	FTimerHandle RespawnTimer;
+
+public:	
+	
+	/** Constructor */
+	AShooterPickup();
+
+protected:
+
+	/** Native construction script */
+	virtual void OnConstruction(const FTransform& Transform) override;
+
+	/** Gameplay Initialization*/
+	virtual void BeginPlay() override;
+
+	/** Gameplay cleanup */
+	virtual void EndPlay(const EEndPlayReason::Type EndPlayReason) override;
+
+	/** Handles collision overlap */
+	UFUNCTION()
+	virtual void OnOverlap(UPrimitiveComponent* OverlappedComponent, AActor* OtherActor, UPrimitiveComponent* OtherComp, int32 OtherBodyIndex, bool bFromSweep, const FHitResult& SweepResult);
+
+protected:
+
+	/** Called when it's time to respawn this pickup */
+	void RespawnPickup();
+
+	/** Passes control to Blueprint to animate the pickup respawn. Should end by calling FinishRespawn */
+	UFUNCTION(BlueprintImplementableEvent, Category="Pickup", meta=(DisplayName = "OnRespawn"))
+	void BP_OnRespawn();
+
+	/** Enables this pickup after respawning */
+	UFUNCTION(BlueprintCallable, Category="Pickup")
+	void FinishRespawn();
+};
diff --git a/Templates/TP_FirstPerson/Source/TP_FirstPerson/Variant_Shooter/ShooterPlayerController.cpp b/Templates/TP_FirstPerson/Source/TP_FirstPerson/Variant_Shooter/ShooterPlayerController.cpp
new file mode 100644
--- /dev/null
+++ b/Templates/TP_FirstPerson/Source/TP_FirstPerson/Variant_Shooter/ShooterPlayerController.cpp
@@ -0,0 +1,78 @@
+// Copyright Epic Games, Inc. All Rights Reserved.
+
+
+#include "Variant_Shooter/ShooterPlayerController.h"
+#include "EnhancedInputSubsystems.h"
+#include "Engine/LocalPlayer.h"
+#include "InputMappingContext.h"
+#include "Kismet/GameplayStatics.h"
+#include "GameFramework/PlayerStart.h"
+#include "ShooterCharacter.h"
+#include "ShooterBulletCounterUI.h"
+
+void AShooterPlayerController::BeginPlay()
+{
+	Super::BeginPlay();
+
+	// create the bullet counter widget and add it to the screen
+	BulletCounterUI = CreateWidget<UShooterBulletCounterUI>(this, BulletCounterUIClass);
+	BulletCounterUI->AddToPlayerScreen(0);
+}
+
+void AShooterPlayerController::SetupInputComponent()
+{
+	// add the input mapping contexts
+	if (UEnhancedInputLocalPlayerSubsystem* Subsystem = ULocalPlayer::GetSubsystem<UEnhancedInputLocalPlayerSubsystem>(GetLocalPlayer()))
+	{
+		for (UInputMappingContext* CurrentContext : DefaultMappingContexts)
+		{
+			Subsystem->AddMappingContext(CurrentContext, 0);
+		}
+	}
+}
+
+void AShooterPlayerController::OnPossess(APawn* InPawn)
+{
+	Super::OnPossess(InPawn);
+
+	// subscribe to the pawn's OnDestroyed delegate
+	InPawn->OnDestroyed.AddDynamic(this, &AShooterPlayerController::OnPawnDestroyed);
+
+	// is this a shooter character?
+	if (AShooterCharacter* ShooterCharacter = Cast<AShooterCharacter>(InPawn))
+	{
+		// add the player tag
+		ShooterCharacter->Tags.Add(PlayerPawnTag);
+
+		// subscribe to the pawn's bullet count updated delegate
+		ShooterCharacter->OnBulletCountUpdated.AddDynamic(this, &AShooterPlayerController::OnBulletCountUpdated);
+	}
+}
+
+void AShooterPlayerController::OnPawnDestroyed(AActor* DestroyedActor)
+{
+	// find the player start
+	TArray<AActor*> ActorList;
+	UGameplayStatics::GetAllActorsOfClass(GetWorld(), APlayerStart::StaticClass(), ActorList);
+
+	if (ActorList.Num() > 0)
+	{
+		// select a random player start
+		AActor* RandomPlayerStart = ActorList[FMath::RandRange(0, ActorList.Num() - 1)];
+
+		// spawn a character at the player start
+		const FTransform SpawnTransform = RandomPlayerStart->GetActorTransform();
+
+		if (AShooterCharacter* RespawnedCharacter = GetWorld()->SpawnActor<AShooterCharacter>(CharacterClass, SpawnTransform))
+		{
+			// possess the character
+			Possess(RespawnedCharacter);
+		}
+	}
+}
+
+void AShooterPlayerController::OnBulletCountUpdated(int32 MagazineSize, int32 Bullets)
+{
+	// update the UI
+	BulletCounterUI->BP_UpdateBulletCounter(MagazineSize, Bullets);
+}
diff --git a/Templates/TP_FirstPerson/Source/TP_FirstPerson/Variant_Shooter/ShooterPlayerController.h b/Templates/TP_FirstPerson/Source/TP_FirstPerson/Variant_Shooter/ShooterPlayerController.h
new file mode 100644
--- /dev/null
+++ b/Templates/TP_FirstPerson/Source/TP_FirstPerson/Variant_Shooter/ShooterPlayerController.h
@@ -0,0 +1,62 @@
+// Copyright Epic Games, Inc. All Rights Reserved.
+
+#pragma once
+
+#include "CoreMinimal.h"
+#include "GameFramework/PlayerController.h"
+#include "ShooterPlayerController.generated.h"
+
+class UInputMappingContext;
+class AShooterCharacter;
+class UShooterBulletCounterUI;
+
+/**
+ *  Simple PlayerController for a first person shooter game
+ *  Manages input mappings
+ *  Respawns the player pawn when it's destroyed
+ */
+UCLASS(abstract)
+class TP_FIRSTPERSON_API AShooterPlayerController : public APlayerController
+{
+	GENERATED_BODY()
+	
+protected:
+
+	/** Input mapping contexts for this player */
+	UPROPERTY(EditAnywhere, Category="Shooter")
+	TArray<UInputMappingContext*> DefaultMappingContexts;
+
+	/** Character class to respawn when the possessed pawn is destroyed */
+	UPROPERTY(EditAnywhere, Category="Shooter")
+	TSubclassOf<AShooterCharacter> CharacterClass;
+
+	/** Type of bullet counter UI widget to spawn */
+	UPROPERTY(EditAnywhere, Category="Shooter")
+	TSubclassOf<UShooterBulletCounterUI> BulletCounterUIClass;
+
+	/** Tag to grant the possessed pawn to flag it as the player */
+	UPROPERTY(EditAnywhere, Category="Shooter")
+	FName PlayerPawnTag = FName("Player");
+
+	/** Pointer to the bullet counter UI widget */
+	TObjectPtr<UShooterBulletCounterUI> BulletCounterUI;
+
+protected:
+
+	/** Gameplay Initialization */
+	virtual void BeginPlay() override;
+
+	/** Initialize input bindings */
+	virtual void SetupInputComponent() override;
+
+	/** Pawn initialization */
+	virtual void OnPossess(APawn* InPawn) override;
+
+	/** Called if the possessed pawn is destroyed */
+	UFUNCTION()
+	void OnPawnDestroyed(AActor* DestroyedActor);
+
+	/** Called when the bullet count on the possessed pawn is updated */
+	UFUNCTION()
+	void OnBulletCountUpdated(int32 MagazineSize, int32 Bullets);
+};
diff --git a/Templates/TP_FirstPerson/Source/TP_FirstPerson/Variant_Shooter/ShooterProjectile.cpp b/Templates/TP_FirstPerson/Source/TP_FirstPerson/Variant_Shooter/ShooterProjectile.cpp
new file mode 100644
--- /dev/null
+++ b/Templates/TP_FirstPerson/Source/TP_FirstPerson/Variant_Shooter/ShooterProjectile.cpp
@@ -0,0 +1,86 @@
+// Copyright Epic Games, Inc. All Rights Reserved.
+
+
+#include "Variant_Shooter/ShooterProjectile.h"
+#include "Components/SphereComponent.h"
+#include "GameFramework/ProjectileMovementComponent.h"
+#include "GameFramework/Character.h"
+#include "Kismet/GameplayStatics.h"
+#include "GameFramework/DamageType.h"
+#include "GameFramework/Pawn.h"
+#include "GameFramework/Controller.h"
+
+AShooterProjectile::AShooterProjectile()
+{
+	PrimaryActorTick.bCanEverTick = true;
+
+	// create the collision component and assign it as the root
+	RootComponent = CollisionComponent = CreateDefaultSubobject<USphereComponent>(TEXT("Collision Component"));
+
+	CollisionComponent->SetSphereRadius(16.0f);
+	CollisionComponent->SetCollisionEnabled(ECollisionEnabled::QueryAndPhysics);
+	CollisionComponent->SetCollisionResponseToAllChannels(ECR_Block);
+	CollisionComponent->CanCharacterStepUpOn = ECanBeCharacterBase::ECB_No;
+
+	// create the projectile movement component. No need to attach it because it's not a Scene Component
+	ProjectileMovement = CreateDefaultSubobject<UProjectileMovementComponent>(TEXT("Projectile Movement"));
+
+	ProjectileMovement->InitialSpeed = 3000.0f;
+	ProjectileMovement->MaxSpeed = 3000.0f;
+	ProjectileMovement->bShouldBounce = true;
+
+	// set the default damage type
+	HitDamageType = UDamageType::StaticClass();
+}
+
+void AShooterProjectile::BeginPlay()
+{
+	Super::BeginPlay();
+	
+	// ignore the pawn that shot this projectile
+	CollisionComponent->IgnoreActorWhenMoving(GetInstigator(), true);
+}
+
+void AShooterProjectile::NotifyHit(class UPrimitiveComponent* MyComp, AActor* Other, class UPrimitiveComponent* OtherComp, bool bSelfMoved, FVector HitLocation, FVector HitNormal, FVector NormalImpulse, const FHitResult& Hit)
+{
+	// ignore if we've already hit something else
+	if (bHit)
+	{
+		return;
+	}
+
+	bHit = true;
+
+	// make AI perception noise
+	MakeNoise(NoiseLoudness, GetInstigator(), GetActorLocation(), NoiseRange, NoiseTag);
+
+	// have we hit a physics object?
+	if (OtherComp->IsSimulatingPhysics())
+	{
+		// give some physics impulse to the object
+		OtherComp->AddImpulseAtLocation(GetVelocity() * PhysicsForce, Hit.ImpactPoint);
+	}
+
+	// have we hit a character?
+	if (ACharacter* HitCharacter = Cast<ACharacter>(Other))
+	{
+		// ignore the owner of this projectile
+		if (HitCharacter != GetOwner())
+		{
+			// damage the hit character
+			DamageCharacter(HitCharacter, Hit);
+		}
+	}
+
+	// disable collision on the projectile
+	CollisionComponent->SetCollisionEnabled(ECollisionEnabled::NoCollision);
+
+	// pass control to BP for any extra effects
+	BP_OnProjectileHit(Hit);
+}
+
+void AShooterProjectile::DamageCharacter(ACharacter* HitCharacter, const FHitResult& Hit)
+{
+	// apply damage to the character
+	UGameplayStatics::ApplyDamage(HitCharacter, HitDamage, GetInstigator()->GetController(), this, HitDamageType);
+}
diff --git a/Templates/TP_FirstPerson/Source/TP_FirstPerson/Variant_Shooter/ShooterProjectile.h b/Templates/TP_FirstPerson/Source/TP_FirstPerson/Variant_Shooter/ShooterProjectile.h
new file mode 100644
--- /dev/null
+++ b/Templates/TP_FirstPerson/Source/TP_FirstPerson/Variant_Shooter/ShooterProjectile.h
@@ -0,0 +1,80 @@
+// Copyright Epic Games, Inc. All Rights Reserved.
+
+#pragma once
+
+#include "CoreMinimal.h"
+#include "GameFramework/Actor.h"
+#include "ShooterProjectile.generated.h"
+
+class USphereComponent;
+class UProjectileMovementComponent;
+class ACharacter;
+
+/**
+ *  Simple projectile class for a first person shooter game
+ */
+UCLASS(abstract)
+class TP_FIRSTPERSON_API AShooterProjectile : public AActor
+{
+	GENERATED_BODY()
+	
+	/** Provides collision detection for the projectile */
+	UPROPERTY(VisibleAnywhere, BlueprintReadOnly, Category="Components", meta = (AllowPrivateAccess = "true"))
+	USphereComponent* CollisionComponent;
+
+	/** Handles movement for the projectile */
+	UPROPERTY(VisibleAnywhere, BlueprintReadOnly, Category="Components", meta = (AllowPrivateAccess = "true"))
+	UProjectileMovementComponent* ProjectileMovement;
+
+protected:
+
+	/** Loudness of the AI perception noise done by this projectile on hit */
+	UPROPERTY(EditAnywhere, Category="Noise")
+	float NoiseLoudness = 3.0f;
+
+	/** Range of the AI perception noise done by this projectile on hit */
+	UPROPERTY(EditAnywhere, Category="Noise")
+	float NoiseRange = 3000.0f;
+
+	/** Tag of the AI perception noise done by this projectile on hit */
+	UPROPERTY(EditAnywhere, Category="Noise")
+	FName NoiseTag = FName("Projectile");
+
+	/** Physics force to apply on hit */
+	UPROPERTY(EditAnywhere, Category="Hit")
+	float PhysicsForce = 100.0f;
+
+	/** Damage to apply on hit */
+	UPROPERTY(EditAnywhere, Category="Hit")
+	float HitDamage = 25.0f;
+
+	/** Type of damage to apply. Can be used to represent specific types of damage such as fire, explosion, etc. */
+	UPROPERTY(EditAnywhere, Category="Hit")
+	TSubclassOf<UDamageType> HitDamageType;
+
+	/** If true, this projectile has already hit another surface */
+	bool bHit = false;
+
+public:	
+
+	/** Constructor */
+	AShooterProjectile();
+
+protected:
+	
+	/** Gameplay initialization */
+	virtual void BeginPlay() override;
+
+	/** Handles collision */
+	virtual void NotifyHit(class UPrimitiveComponent* MyComp, AActor* Other, class UPrimitiveComponent* OtherComp, bool bSelfMoved, FVector HitLocation, FVector HitNormal, FVector NormalImpulse, const FHitResult& Hit) override;
+
+protected:
+
+	/** Apply damage to a hit character */
+	UFUNCTION(BlueprintCallable, Category="Projectile")
+	virtual void DamageCharacter(ACharacter* HitCharacter, const FHitResult& Hit);
+
+	/** Passes control to Blueprint to implement any effects on hit */
+	UFUNCTION(BlueprintImplementableEvent, Category="Projectile", meta=(DisplayName = "On Projectile Hit"))
+	void BP_OnProjectileHit(const FHitResult& Hit);
+};
diff --git a/Templates/TP_FirstPerson/Source/TP_FirstPerson/Variant_Shooter/ShooterUI.cpp b/Templates/TP_FirstPerson/Source/TP_FirstPerson/Variant_Shooter/ShooterUI.cpp
new file mode 100644
--- /dev/null
+++ b/Templates/TP_FirstPerson/Source/TP_FirstPerson/Variant_Shooter/ShooterUI.cpp
@@ -0,0 +1,5 @@
+// Copyright Epic Games, Inc. All Rights Reserved.
+
+
+#include "Variant_Shooter/ShooterUI.h"
+
diff --git a/Templates/TP_FirstPerson/Source/TP_FirstPerson/Variant_Shooter/ShooterUI.h b/Templates/TP_FirstPerson/Source/TP_FirstPerson/Variant_Shooter/ShooterUI.h
new file mode 100644
--- /dev/null
+++ b/Templates/TP_FirstPerson/Source/TP_FirstPerson/Variant_Shooter/ShooterUI.h
@@ -0,0 +1,22 @@
+// Copyright Epic Games, Inc. All Rights Reserved.
+
+#pragma once
+
+#include "CoreMinimal.h"
+#include "Blueprint/UserWidget.h"
+#include "ShooterUI.generated.h"
+
+/**
+ *  Simple scoreboard UI for a first person shooter game
+ */
+UCLASS(abstract)
+class TP_FIRSTPERSON_API UShooterUI : public UUserWidget
+{
+	GENERATED_BODY()
+	
+public:
+
+	/** Allows Blueprint to update score sub-widgets */
+	UFUNCTION(BlueprintImplementableEvent, Category="Shooter", meta=(DisplayName = "Update Score"))
+	void BP_UpdateScore(uint8 TeamByte, int32 Score);
+};
diff --git a/Templates/TP_FirstPerson/Source/TP_FirstPerson/Variant_Shooter/ShooterWeapon.cpp b/Templates/TP_FirstPerson/Source/TP_FirstPerson/Variant_Shooter/ShooterWeapon.cpp
new file mode 100644
--- /dev/null
+++ b/Templates/TP_FirstPerson/Source/TP_FirstPerson/Variant_Shooter/ShooterWeapon.cpp
@@ -0,0 +1,218 @@
+// Copyright Epic Games, Inc. All Rights Reserved.
+
+
+#include "Variant_Shooter/ShooterWeapon.h"
+#include "Kismet/KismetMathLibrary.h"
+#include "Engine/World.h"
+#include "ShooterProjectile.h"
+#include "ShooterWeaponHolder.h"
+#include "Components/SceneComponent.h"
+#include "TimerManager.h"
+#include "Animation/AnimInstance.h"
+#include "Components/SkeletalMeshComponent.h"
+#include "GameFramework/Pawn.h"
+
+AShooterWeapon::AShooterWeapon()
+{
+	PrimaryActorTick.bCanEverTick = true;
+
+	// create the root
+	RootComponent = CreateDefaultSubobject<USceneComponent>(TEXT("Root"));
+
+	// create the first person mesh
+	FirstPersonMesh = CreateDefaultSubobject<USkeletalMeshComponent>(TEXT("First Person Mesh"));
+	FirstPersonMesh->SetupAttachment(RootComponent);
+
+	FirstPersonMesh->SetCollisionProfileName(FName("NoCollision"));
+	FirstPersonMesh->SetFirstPersonPrimitiveType(EFirstPersonPrimitiveType::FirstPerson);
+	FirstPersonMesh->bOnlyOwnerSee = true;
+
+	// create the third person mesh
+	ThirdPersonMesh = CreateDefaultSubobject<USkeletalMeshComponent>(TEXT("Third Person Mesh"));
+	ThirdPersonMesh->SetupAttachment(RootComponent);
+
+	ThirdPersonMesh->SetCollisionProfileName(FName("NoCollision"));
+	ThirdPersonMesh->SetFirstPersonPrimitiveType(EFirstPersonPrimitiveType::WorldSpaceRepresentation);
+	ThirdPersonMesh->bOwnerNoSee = true;
+}
+
+void AShooterWeapon::BeginPlay()
+{
+	Super::BeginPlay();
+
+	// subscribe to the owner's destroyed delegate
+	GetOwner()->OnDestroyed.AddDynamic(this, &AShooterWeapon::OnOwnerDestroyed);
+
+	// cast the weapon owner
+	WeaponOwner = Cast<IShooterWeaponHolder>(GetOwner());
+	PawnOwner = Cast<APawn>(GetOwner());
+
+	// fill the first ammo clip
+	CurrentBullets = MagazineSize;
+
+	// attach the meshes to the owner
+	WeaponOwner->AttachWeaponMeshes(this);
+}
+
+void AShooterWeapon::EndPlay(EEndPlayReason::Type EndPlayReason)
+{
+	Super::EndPlay(EndPlayReason);
+
+	// clear the refire timer
+	GetWorld()->GetTimerManager().ClearTimer(RefireTimer);
+}
+
+void AShooterWeapon::OnOwnerDestroyed(AActor* DestroyedActor)
+{
+	// ensure this weapon is destroyed when the owner is destroyed
+	Destroy();
+}
+
+void AShooterWeapon::ActivateWeapon()
+{
+	// unhide this weapon
+	SetActorHiddenInGame(false);
+
+	// notify the owner
+	WeaponOwner->OnWeaponActivated(this);
+}
+
+void AShooterWeapon::DeactivateWeapon()
+{
+	// ensure we're no longer firing this weapon while deactivated
+	StopFiring();
+
+	// hide the weapon
+	SetActorHiddenInGame(true);
+
+	// notify the owner
+	WeaponOwner->OnWeaponDeactivated(this);
+}
+
+void AShooterWeapon::StartFiring()
+{
+	// raise the firing flag
+	bIsFiring = true;
+
+	// check how much time has passed since we last shot
+	// this may be under the refire rate if the weapon shoots slow enough and the player is spamming the trigger
+	const float TimeSinceLastShot = GetWorld()->GetTimeSeconds() - TimeOfLastShot;
+
+	if (TimeSinceLastShot > RefireRate)
+	{
+		// fire the weapon right away
+		Fire();
+
+	} else {
+
+		// if we're full auto, schedule the next shot
+		if (bFullAuto)
+		{
+			GetWorld()->GetTimerManager().SetTimer(RefireTimer, this, &AShooterWeapon::Fire, TimeSinceLastShot, false);
+		}
+
+	}
+}
+
+void AShooterWeapon::StopFiring()
+{
+	// lower the firing flag
+	bIsFiring = false;
+
+	// clear the refire timer
+	GetWorld()->GetTimerManager().ClearTimer(RefireTimer);
+}
+
+void AShooterWeapon::Fire()
+{
+	// ensure the player still wants to fire. They may have let go of the trigger
+	if (!bIsFiring)
+	{
+		return;
+	}
+	
+	// fire a projectile at the target
+	FireProjectile(WeaponOwner->GetWeaponTargetLocation());
+
+	// update the time of our last shot
+	TimeOfLastShot = GetWorld()->GetTimeSeconds();
+
+	// make noise so the AI perception system can hear us
+	MakeNoise(ShotLoudness, PawnOwner, PawnOwner->GetActorLocation(), ShotNoiseRange, ShotNoiseTag);
+
+	// are we full auto?
+	if (bFullAuto)
+	{
+		// schedule the next shot
+		GetWorld()->GetTimerManager().SetTimer(RefireTimer, this, &AShooterWeapon::Fire, RefireRate, false);
+	} else {
+
+		// for semi-auto weapons, schedule the cooldown notification
+		GetWorld()->GetTimerManager().SetTimer(RefireTimer, this, &AShooterWeapon::FireCooldownExpired, RefireRate, false);
+
+	}
+}
+
+void AShooterWeapon::FireCooldownExpired()
+{
+	// notify the owner
+	WeaponOwner->OnSemiWeaponRefire();
+}
+
+void AShooterWeapon::FireProjectile(const FVector& TargetLocation)
+{
+	// get the projectile transform
+	FTransform ProjectileTransform = CalculateProjectileSpawnTransform(TargetLocation);
+	
+	// spawn the projectile
+	FActorSpawnParameters SpawnParams;
+	SpawnParams.SpawnCollisionHandlingOverride = ESpawnActorCollisionHandlingMethod::AlwaysSpawn;
+	SpawnParams.TransformScaleMethod = ESpawnActorScaleMethod::OverrideRootScale;
+	SpawnParams.Owner = GetOwner();
+	SpawnParams.Instigator = PawnOwner;
+
+	AShooterProjectile* Projectile = GetWorld()->SpawnActor<AShooterProjectile>(ProjectileClass, ProjectileTransform, SpawnParams);
+
+	// play the firing montage
+	WeaponOwner->PlayFiringMontage(FiringMontage);
+
+	// add recoil
+	WeaponOwner->AddWeaponRecoil(FiringRecoil);
+
+	// consume bullets
+	--CurrentBullets;
+
+	// if the clip is depleted, reload it
+	if (CurrentBullets <= 0)
+	{
+		CurrentBullets = MagazineSize;
+	}
+
+	// update the weapon HUD
+	WeaponOwner->UpdateWeaponHUD(CurrentBullets, MagazineSize);
+}
+
+FTransform AShooterWeapon::CalculateProjectileSpawnTransform(const FVector& TargetLocation) const
+{
+	// find the muzzle location
+	const FVector MuzzleLoc = FirstPersonMesh->GetSocketLocation(MuzzleSocketName);
+
+	// calculate the spawn location ahead of the muzzle
+	const FVector SpawnLoc = MuzzleLoc + ((TargetLocation - MuzzleLoc).GetSafeNormal() * MuzzleOffset);
+
+	// find the aim rotation vector while applying some variance to the target 
+	const FRotator AimRot = UKismetMathLibrary::FindLookAtRotation(SpawnLoc, TargetLocation + (UKismetMathLibrary::RandomUnitVector() * AimVariance));
+
+	// return the built transform
+	return FTransform(AimRot, SpawnLoc, FVector::OneVector);
+}
+
+const TSubclassOf<UAnimInstance>& AShooterWeapon::GetFirstPersonAnimInstanceClass() const
+{
+	return FirstPersonAnimInstanceClass;
+}
+
+const TSubclassOf<UAnimInstance>& AShooterWeapon::GetThirdPersonAnimInstanceClass() const
+{
+	return ThirdPersonAnimInstanceClass;
+}
diff --git a/Templates/TP_FirstPerson/Source/TP_FirstPerson/Variant_Shooter/ShooterWeapon.h b/Templates/TP_FirstPerson/Source/TP_FirstPerson/Variant_Shooter/ShooterWeapon.h
new file mode 100644
--- /dev/null
+++ b/Templates/TP_FirstPerson/Source/TP_FirstPerson/Variant_Shooter/ShooterWeapon.h
@@ -0,0 +1,180 @@
+// Copyright Epic Games, Inc. All Rights Reserved.
+
+#pragma once
+
+#include "CoreMinimal.h"
+#include "GameFramework/Actor.h"
+#include "ShooterWeaponHolder.h"
+#include "Animation/AnimInstance.h"
+#include "ShooterWeapon.generated.h"
+
+class IShooterWeaponHolder;
+class AShooterProjectile;
+class USkeletalMeshComponent;
+class UAnimMontage;
+class UAnimInstance;
+
+/**
+ *  Base class for a simple first person shooter weapon
+ *  Provides both first person and third person perspective meshes
+ *  Handles ammo and firing logic
+ *  Interacts with the weapon owner through the ShooterWeaponHolder interface
+ */
+UCLASS(abstract)
+class TP_FIRSTPERSON_API AShooterWeapon : public AActor
+{
+	GENERATED_BODY()
+	
+	/** First person perspective mesh */
+	UPROPERTY(VisibleAnywhere, BlueprintReadOnly, Category="Components", meta = (AllowPrivateAccess = "true"))
+	USkeletalMeshComponent* FirstPersonMesh;
+
+	/** Third person perspective mesh */
+	UPROPERTY(VisibleAnywhere, BlueprintReadOnly, Category="Components", meta = (AllowPrivateAccess = "true"))
+	USkeletalMeshComponent* ThirdPersonMesh;
+
+protected:
+
+	/** Cast pointer to the weapon owner */
+	IShooterWeaponHolder* WeaponOwner;
+
+	/** Type of projectiles this weapon will shoot */
+	UPROPERTY(EditAnywhere, Category="Ammo")
+	TSubclassOf<AShooterProjectile> ProjectileClass;
+
+	/** Number of bullets in a magazine */
+	UPROPERTY(EditAnywhere, Category="Ammo")
+	int32 MagazineSize = 10;
+
+	/** Number of bullets in the current magazine */
+	int32 CurrentBullets = 0;
+	
+	/** Animation montage to play when firing this weapon */
+	UPROPERTY(EditAnywhere, Category="Animation")
+	UAnimMontage* FiringMontage;
+
+	/** AnimInstance class to set for the first person character mesh when this weapon is active */
+	UPROPERTY(EditAnywhere, Category="Animation")
+	TSubclassOf<UAnimInstance> FirstPersonAnimInstanceClass;
+
+	/** AnimInstance class to set for the third person character mesh when this weapon is active */
+	UPROPERTY(EditAnywhere, Category="Animation")
+	TSubclassOf<UAnimInstance> ThirdPersonAnimInstanceClass;
+
+	/** Cone half-angle for variance while aiming */
+	UPROPERTY(EditAnywhere, Category="Aim")
+	float AimVariance = 0.0f;
+
+	/** Amount of firing recoil to apply to the owner */
+	UPROPERTY(EditAnywhere, Category="Aim")
+	float FiringRecoil = 0.0f;
+
+	/** Name of the first person muzzle socket where projectiles will spawn */
+	UPROPERTY(EditAnywhere, Category="Aim")
+	FName MuzzleSocketName;
+
+	/** Distance ahead of the muzzle that bullets will spawn at */
+	UPROPERTY(EditAnywhere, Category="Aim")
+	float MuzzleOffset = 10.0f;
+
+	/** If true, this weapon will automatically fire at the refire rate */
+	UPROPERTY(EditAnywhere, Category="Refire")
+	bool bFullAuto = false;
+
+	/** Time between shots for this weapon. Affects both full auto and semi auto modes */
+	UPROPERTY(EditAnywhere, Category="Refire")
+	float RefireRate = 0.5f;
+
+	/** Game time of last shot fired, used to enforce refire rate on semi auto */
+	float TimeOfLastShot = 0.0f;
+
+	/** If true, the weapon is currently firing */
+	bool bIsFiring = false;
+
+	/** Timer to handle full auto refiring */
+	FTimerHandle RefireTimer;
+
+	/** Cast pawn pointer to the owner for AI perception system interactions */
+	TObjectPtr<APawn> PawnOwner;
+
+	/** Loudness of the shot for AI perception system interactions */
+	UPROPERTY(EditAnywhere, Category="Perception")
+	float ShotLoudness = 1.0f;
+
+	/** Max range of shot AI perception noise */
+	UPROPERTY(EditAnywhere, Category="Perception")
+	float ShotNoiseRange = 3000.0f;
+
+	/** Tag to apply to noise generated by shooting this weapon */
+	UPROPERTY(EditAnywhere, Category="Perception")
+	FName ShotNoiseTag = FName("Shot");
+
+public:	
+
+	/** Constructor */
+	AShooterWeapon();
+
+protected:
+	
+	/** Gameplay initialization */
+	virtual void BeginPlay() override;
+
+	/** Gameplay Cleanup */
+	virtual void EndPlay(EEndPlayReason::Type EndPlayReason) override;
+
+protected:
+
+	/** Called when the weapon's owner is destroyed */
+	UFUNCTION()
+	void OnOwnerDestroyed(AActor* DestroyedActor);
+
+public:
+
+	/** Activates this weapon and gets it ready to fire */
+	void ActivateWeapon();
+
+	/** Deactivates this weapon */
+	void DeactivateWeapon();
+
+	/** Start firing this weapon */
+	void StartFiring();
+
+	/** Stop firing this weapon */
+	void StopFiring();
+
+protected:
+
+	/** Fire the weapon */
+	virtual void Fire();
+
+	/** Called when the refire rate time has passed while shooting semi auto weapons */
+	void FireCooldownExpired();
+
+	/** Fire a projectile towards the target location */
+	virtual void FireProjectile(const FVector& TargetLocation);
+
+	/** Calculates the spawn transform for projectiles shot by this weapon */
+	FTransform CalculateProjectileSpawnTransform(const FVector& TargetLocation) const;
+
+public:
+
+	/** Returns the first person mesh */
+	UFUNCTION(BlueprintPure, Category="Weapon")
+	USkeletalMeshComponent* GetFirstPersonMesh() const { return FirstPersonMesh; };
+
+	/** Returns the third person mesh */
+	UFUNCTION(BlueprintPure, Category="Weapon")
+	USkeletalMeshComponent* GetThirdPersonMesh() const { return ThirdPersonMesh; };
+
+	/** Returns the first person anim instance class */
+	const TSubclassOf<UAnimInstance>& GetFirstPersonAnimInstanceClass() const;
+
+	/** Returns the third person anim instance class */
+	const TSubclassOf<UAnimInstance>& GetThirdPersonAnimInstanceClass() const;
+
+	/** Returns the magazine size */
+	int32 GetMagazineSize() const { return MagazineSize; };
+
+	/** Returns the current bullet count */
+	int32 GetBulletCount() const { return CurrentBullets; }
+};
diff --git a/Templates/TP_FirstPerson/Source/TP_FirstPerson/Variant_Shooter/ShooterWeaponHolder.cpp b/Templates/TP_FirstPerson/Source/TP_FirstPerson/Variant_Shooter/ShooterWeaponHolder.cpp
new file mode 100644
--- /dev/null
+++ b/Templates/TP_FirstPerson/Source/TP_FirstPerson/Variant_Shooter/ShooterWeaponHolder.cpp
@@ -0,0 +1,6 @@
+// Copyright Epic Games, Inc. All Rights Reserved.
+
+
+#include "Variant_Shooter/ShooterWeaponHolder.h"
+
+// Add default functionality here for any IShooterWeaponHolder functions that are not pure virtual.
diff --git a/Templates/TP_FirstPerson/Source/TP_FirstPerson/Variant_Shooter/ShooterWeaponHolder.h b/Templates/TP_FirstPerson/Source/TP_FirstPerson/Variant_Shooter/ShooterWeaponHolder.h
new file mode 100644
--- /dev/null
+++ b/Templates/TP_FirstPerson/Source/TP_FirstPerson/Variant_Shooter/ShooterWeaponHolder.h
@@ -0,0 +1,55 @@
+// Copyright Epic Games, Inc. All Rights Reserved.
+
+#pragma once
+
+#include "CoreMinimal.h"
+#include "UObject/Interface.h"
+#include "ShooterWeaponHolder.generated.h"
+
+class AShooterWeapon;
+class UAnimMontage;
+
+
+// This class does not need to be modified.
+UINTERFACE(MinimalAPI)
+class UShooterWeaponHolder : public UInterface
+{
+	GENERATED_BODY()
+};
+
+/**
+ *  Common interface for Shooter Game weapon holder classes
+ */
+class TP_FIRSTPERSON_API IShooterWeaponHolder
+{
+	GENERATED_BODY()
+
+public:
+
+	/** Attaches a weapon's meshes to the owner */
+	virtual void AttachWeaponMeshes(AShooterWeapon* Weapon) = 0;
+
+	/** Plays the firing montage for the weapon */
+	virtual void PlayFiringMontage(UAnimMontage* Montage) = 0;
+
+	/** Applies weapon recoil to the owner */
+	virtual void AddWeaponRecoil(float Recoil) = 0;
+
+	/** Updates the weapon's HUD with the current ammo count */
+	virtual void UpdateWeaponHUD(int32 CurrentAmmo, int32 MagazineSize) = 0;
+
+	/** Calculates and returns the aim location for the weapon */
+	virtual FVector GetWeaponTargetLocation() = 0;
+
+	/** Gives a weapon of this class to the owner */
+	virtual void AddWeaponClass(const TSubclassOf<AShooterWeapon>& WeaponClass) = 0;
+
+	/** Activates the passed weapon */
+	virtual void OnWeaponActivated(AShooterWeapon* Weapon) = 0;
+
+	/** Deactivates the passed weapon */
+	virtual void OnWeaponDeactivated(AShooterWeapon* Weapon) = 0;
+
+	/** Notifies the owner that the weapon cooldown has expired and it's ready to shoot again */
+	virtual void OnSemiWeaponRefire() = 0;
+};
diff --git a/Templates/TP_FirstPerson/TP_FirstPerson.uproject b/Templates/TP_FirstPerson/TP_FirstPerson.uproject
--- a/Templates/TP_FirstPerson/TP_FirstPerson.uproject
+++ b/Templates/TP_FirstPerson/TP_FirstPerson.uproject
@@ -7,7 +7,12 @@
 		{
 			"Name": "TP_FirstPerson",
 			"Type": "Runtime",
-			"LoadingPhase": "Default"
+			"LoadingPhase": "Default",
+			"AdditionalDependencies": [
+				"Engine",
+				"AIModule",
+				"UMG"
+			]
 		}
 	],
 	"Plugins": [
@@ -17,6 +22,14 @@
 			"TargetAllowList": [
 				"Editor"
 			]
+		},
+		{
+			"Name": "StateTree",
+			"Enabled": true
+		},
+		{
+			"Name": "GameplayStateTree",
+			"Enabled": true
 		}
-	]	
+	]
 }
\ No newline at end of file
