TP_FirstPerson - Fixed crash due to attempting to grant unloaded weapon class to player on Shooter variant.
#rnx

diff --git a/Templates/TP_FirstPerson/Source/TP_FirstPerson/Variant_Shooter/ShooterPickup.cpp b/Templates/TP_FirstPerson/Source/TP_FirstPerson/Variant_Shooter/ShooterPickup.cpp
--- a/Templates/TP_FirstPerson/Source/TP_FirstPerson/Variant_Shooter/ShooterPickup.cpp
+++ b/Templates/TP_FirstPerson/Source/TP_FirstPerson/Variant_Shooter/ShooterPickup.cpp
@@ -56,7 +56,7 @@ void AShooterPickup::BeginPlay()
 	if (FWeaponTableRow* WeaponData = WeaponType.GetRow<FWeaponTableRow>(FString()))
 	{
 		// copy the weapon class
-		WeaponClass = WeaponData->WeaponClass.LoadSynchronous();
+		WeaponClass = WeaponData->WeaponToSpawn;
 	}
 }
 
diff --git a/Templates/TP_FirstPerson/Source/TP_FirstPerson/Variant_Shooter/ShooterPickup.h b/Templates/TP_FirstPerson/Source/TP_FirstPerson/Variant_Shooter/ShooterPickup.h
--- a/Templates/TP_FirstPerson/Source/TP_FirstPerson/Variant_Shooter/ShooterPickup.h
+++ b/Templates/TP_FirstPerson/Source/TP_FirstPerson/Variant_Shooter/ShooterPickup.h
@@ -26,7 +26,7 @@ struct FWeaponTableRow : public FTableRowBase
 
 	/** Weapon class to grant on pickup */
 	UPROPERTY(EditAnywhere)
-	TSoftClassPtr<AShooterWeapon> WeaponClass;
+	TSubclassOf<AShooterWeapon> WeaponToSpawn;
 };
 
 /**
