[State Tree] fixed calling GetInstanceDataPtr from rvalue reference StrongExecContext compile error

diff --git a/Templates/TP_FirstPerson/Source/TP_FirstPerson/Variant_Shooter/AI/ShooterStateTreeUtility.cpp b/Templates/TP_FirstPerson/Source/TP_FirstPerson/Variant_Shooter/AI/ShooterStateTreeUtility.cpp
--- a/Templates/TP_FirstPerson/Source/TP_FirstPerson/Variant_Shooter/AI/ShooterStateTreeUtility.cpp
+++ b/Templates/TP_FirstPerson/Source/TP_FirstPerson/Variant_Shooter/AI/ShooterStateTreeUtility.cpp
@@ -224,69 +224,67 @@ EStateTreeRunStatus FStateTreeSenseEnemiesTask::EnterState(FStateTreeExecutionCo
 			[WeakContext = Context.MakeWeakExecutionContext()](AActor* SensedActor, const FAIStimulus& Stimulus)
 			{
 				// get the instance data inside the lambda
-				FInstanceDataType* LambdaInstanceData = WeakContext.MakeStrongExecutionContext().GetInstanceDataPtr<FInstanceDataType>();
-
-				if (!LambdaInstanceData)
+				const FStateTreeStrongExecutionContext StrongContext = WeakContext.MakeStrongExecutionContext();
+				if (FInstanceDataType* LambdaInstanceData = StrongContext.GetInstanceDataPtr<FInstanceType>())
 				{
-					return;
-				}
-
-				if (SensedActor->ActorHasTag(LambdaInstanceData->SenseTag))
-				{
-					bool bDirectLOS = false;
-
-					// calculate the direction of the stimulus
-					const FVector StimulusDir = (Stimulus.StimulusLocation - LambdaInstanceData->Character->GetActorLocation()).GetSafeNormal();
+					if (SensedActor->ActorHasTag(LambdaInstanceData->SenseTag))
+					{
+						bool bDirectLOS = false;
 
-					// infer the angle from the dot product between the character facing and the stimulus direction
-					const float DirDot = FVector::DotProduct(StimulusDir, LambdaInstanceData->Character->GetActorForwardVector());
-					const float MaxDot = FMath::Cos(FMath::DegreesToRadians(LambdaInstanceData->DirectLineOfSightCone));
+						// calculate the direction of the stimulus
+						const FVector StimulusDir = (Stimulus.StimulusLocation - LambdaInstanceData->Character->GetActorLocation()).GetSafeNormal();
 
-					// is the direction within our perception cone?
-					if (DirDot >= MaxDot)
-					{
-						// run a line trace between the character and the sensed actor
-						FCollisionQueryParams QueryParams;
-						QueryParams.AddIgnoredActor(LambdaInstanceData->Character);
-						QueryParams.AddIgnoredActor(SensedActor);
+						// infer the angle from the dot product between the character facing and the stimulus direction
+						const float DirDot = FVector::DotProduct(StimulusDir, LambdaInstanceData->Character->GetActorForwardVector());
+						const float MaxDot = FMath::Cos(FMath::DegreesToRadians(LambdaInstanceData->DirectLineOfSightCone));
 
-						FHitResult OutHit;
+						// is the direction within our perception cone?
+						if (DirDot >= MaxDot)
+						{
+							// run a line trace between the character and the sensed actor
+							FCollisionQueryParams QueryParams;
+							QueryParams.AddIgnoredActor(LambdaInstanceData->Character);
+							QueryParams.AddIgnoredActor(SensedActor);
 
-						// we have direct line of sight if this trace is unobstructed
-						bDirectLOS = !LambdaInstanceData->Character->GetWorld()->LineTraceSingleByChannel(OutHit, LambdaInstanceData->Character->GetActorLocation(), SensedActor->GetActorLocation(), ECC_Visibility, QueryParams);
+							FHitResult OutHit;
 
-					}
+							// we have direct line of sight if this trace is unobstructed
+							bDirectLOS = !LambdaInstanceData->Character->GetWorld()->LineTraceSingleByChannel(OutHit, LambdaInstanceData->Character->GetActorLocation(), SensedActor->GetActorLocation(), ECC_Visibility, QueryParams);
 
-					// check if we have a direct line of sight to the stimulus
-					if (bDirectLOS)
-					{
-						// set the controller's target
-						LambdaInstanceData->Controller->SetCurrentTarget(SensedActor);
+						}
 
-						// set the task output
-						LambdaInstanceData->TargetActor = SensedActor;
+						// check if we have a direct line of sight to the stimulus
+						if (bDirectLOS)
+						{
+							// set the controller's target
+							LambdaInstanceData->Controller->SetCurrentTarget(SensedActor);
 
-						// set the flags
-						LambdaInstanceData->bHasTarget = true;
-						LambdaInstanceData->bHasInvestigateLocation = false;
+							// set the task output
+							LambdaInstanceData->TargetActor = SensedActor;
 
-					// no direct line of sight to target
-					} else {
+							// set the flags
+							LambdaInstanceData->bHasTarget = true;
+							LambdaInstanceData->bHasInvestigateLocation = false;
 
-						// if we already have a target, ignore the partial sense and keep on them
-						if (!IsValid(LambdaInstanceData->TargetActor))
+							// no direct line of sight to target
+						}
+						else
 						{
-							// is this stimulus stronger than the last one we had?
-							if (Stimulus.Strength > LambdaInstanceData->LastStimulusStrength)
+							// if we already have a target, ignore the partial sense and keep on them
+							if (!IsValid(LambdaInstanceData->TargetActor))
 							{
-								// update the stimulus strength
-								LambdaInstanceData->LastStimulusStrength = Stimulus.Strength;
-
-								// set the investigate location
-								LambdaInstanceData->InvestigateLocation = Stimulus.StimulusLocation;
-
-								// set the investigate flag
-								LambdaInstanceData->bHasInvestigateLocation = true;
+								// is this stimulus stronger than the last one we had?
+								if (Stimulus.Strength > LambdaInstanceData->LastStimulusStrength)
+								{
+									// update the stimulus strength
+									LambdaInstanceData->LastStimulusStrength = Stimulus.Strength;
+
+									// set the investigate location
+									LambdaInstanceData->InvestigateLocation = Stimulus.StimulusLocation;
+
+									// set the investigate flag
+									LambdaInstanceData->bHasInvestigateLocation = true;
+								}
 							}
 						}
 					}
@@ -299,46 +297,42 @@ EStateTreeRunStatus FStateTreeSenseEnemiesTask::EnterState(FStateTreeExecutionCo
 			[WeakContext = Context.MakeWeakExecutionContext()](AActor* SensedActor)
 			{
 				// get the instance data inside the lambda
-				FInstanceDataType* LambdaInstanceData = WeakContext.MakeStrongExecutionContext().GetInstanceDataPtr<FInstanceDataType>();
-
-				if (!LambdaInstanceData)
-				{
-					return;
-				}
-
-				bool bForget = false;
-
-				// are we forgetting the current target?
-				if (SensedActor == LambdaInstanceData->TargetActor)
+				const FStateTreeStrongExecutionContext StrongContext = WeakContext.MakeStrongExecutionContext();
+				if (FInstanceDataType* LambdaInstanceData = StrongContext.GetInstanceDataPtr<FInstanceDataType>())
 				{
-					bForget = true;
+					bool bForget = false;
 
-				} else {
-
-					// are we forgetting about a partial sense?
-					if (!IsValid(LambdaInstanceData->TargetActor))
+					// are we forgetting the current target?
+					if (SensedActor == LambdaInstanceData->TargetActor)
 					{
 						bForget = true;
 					}
-				}
+					else 
+					{
+						// are we forgetting about a partial sense?
+						if (!IsValid(LambdaInstanceData->TargetActor))
+						{
+							bForget = true;
+						}
+					}
 
-				if (bForget)
-				{
-					// clear the target
-					LambdaInstanceData->TargetActor = nullptr;
+					if (bForget)
+					{
+						// clear the target
+						LambdaInstanceData->TargetActor = nullptr;
 
-					// clear the flags
-					LambdaInstanceData->bHasInvestigateLocation = false;
-					LambdaInstanceData->bHasTarget = false;
+						// clear the flags
+						LambdaInstanceData->bHasInvestigateLocation = false;
+						LambdaInstanceData->bHasTarget = false;
 
-					// reset the stimulus strength
-					LambdaInstanceData->LastStimulusStrength = 0.0f;
+						// reset the stimulus strength
+						LambdaInstanceData->LastStimulusStrength = 0.0f;
 
-					// clear the target on the controller
-					LambdaInstanceData->Controller->ClearCurrentTarget();
-					LambdaInstanceData->Controller->ClearFocus(EAIFocusPriority::Gameplay);
+						// clear the target on the controller
+						LambdaInstanceData->Controller->ClearCurrentTarget();
+						LambdaInstanceData->Controller->ClearFocus(EAIFocusPriority::Gameplay);
+					}
 				}
-
 			}
 		);
 	}
