[State Tree] fixed typo in ShooterStateTreeUtility file

diff --git a/Templates/TP_FirstPerson/Source/TP_FirstPerson/Variant_Shooter/AI/ShooterStateTreeUtility.cpp b/Templates/TP_FirstPerson/Source/TP_FirstPerson/Variant_Shooter/AI/ShooterStateTreeUtility.cpp
--- a/Templates/TP_FirstPerson/Source/TP_FirstPerson/Variant_Shooter/AI/ShooterStateTreeUtility.cpp
+++ b/Templates/TP_FirstPerson/Source/TP_FirstPerson/Variant_Shooter/AI/ShooterStateTreeUtility.cpp
@@ -225,7 +225,7 @@ EStateTreeRunStatus FStateTreeSenseEnemiesTask::EnterState(FStateTreeExecutionCo
 			{
 				// get the instance data inside the lambda
 				const FStateTreeStrongExecutionContext StrongContext = WeakContext.MakeStrongExecutionContext();
-				if (FInstanceDataType* LambdaInstanceData = StrongContext.GetInstanceDataPtr<FInstanceType>())
+				if (FInstanceDataType* LambdaInstanceData = StrongContext.GetInstanceDataPtr<FInstanceDataType>())
 				{
 					if (SensedActor->ActorHasTag(LambdaInstanceData->SenseTag))
 					{
