Templates
- TP_FirstPerson: Adjusted physics impulse applied to enemy ragdolls
- TP_FirstPerson: Implemented deferred character death
- TP_FirstPerson: Grenades now deal AoE damage and can deal self damage
- TP_FirstPerson: Added life bar and damage feedback to HUD
- TP_FirstPerson: Fixed incorrect HUD scoreboard display
- TP_FirstPerson: Disabled jump on Horror variant
- TP_FirstPerson: Adjusted mass of physics boxes on Shooter map
- TP_FirstPerson: Sprint meter no longer drains when standing still on horror variant
- TP_FirstPersonBP: Renamed and reorganized Blueprint assets
- TP_FirstPersonBP: Changed light mobility to Movable on Horror light BPs
- TP_FirstPersonBP: Disabled jump on Horror variant
- TP_FirstPersonBP: Adjusted physics impulse applied to enemy ragdolls
- TP_FirstPersonBP: Implemented deferred character death
- TP_FirstPersonBP: Grenades now deal AoE damage and can deal self damage
- TP_FirstPersonBP: Added life bar and damage feedback to HUD
- TP_FirstPersonBP: Removed unneeded Asset Guidelines from variants
- TP_FirstPersonBP: Adjusted mass of physics boxes on Shooter map
- TP_FirstPersonBP: Implemented a more efficient sprint meter calculation on Horror variant
- TP_FirstPersonBP: Sprint meter no longer drains when standing still on horror variant

#rnx

diff --git a/Templates/TP_FirstPerson/Source/TP_FirstPerson/Variant_Horror/HorrorCharacter.cpp b/Templates/TP_FirstPerson/Source/TP_FirstPerson/Variant_Horror/HorrorCharacter.cpp
--- a/Templates/TP_FirstPerson/Source/TP_FirstPerson/Variant_Horror/HorrorCharacter.cpp
+++ b/Templates/TP_FirstPerson/Source/TP_FirstPerson/Variant_Horror/HorrorCharacter.cpp
@@ -97,8 +97,8 @@ void AHorrorCharacter::DoEndSprint()
 
 void AHorrorCharacter::SprintFixedTick()
 {
-	// are we out of recovery and still have stamina?
-	if (bSprinting && !bRecovering)
+	// are we out of recovery, still have stamina and are moving faster than our walk speed?
+	if (bSprinting && !bRecovering && GetVelocity().Length() > WalkSpeed)
 	{
 
 		// do we still have meter to burn?
diff --git a/Templates/TP_FirstPerson/Source/TP_FirstPerson/Variant_Shooter/ShooterCharacter.cpp b/Templates/TP_FirstPerson/Source/TP_FirstPerson/Variant_Shooter/ShooterCharacter.cpp
--- a/Templates/TP_FirstPerson/Source/TP_FirstPerson/Variant_Shooter/ShooterCharacter.cpp
+++ b/Templates/TP_FirstPerson/Source/TP_FirstPerson/Variant_Shooter/ShooterCharacter.cpp
@@ -10,6 +10,8 @@
 #include "Components/SkeletalMeshComponent.h"
 #include "Engine/World.h"
 #include "Camera/CameraComponent.h"
+#include "TimerManager.h"
+#include "ShooterGameMode.h"
 
 AShooterCharacter::AShooterCharacter()
 {
@@ -20,6 +22,25 @@ AShooterCharacter::AShooterCharacter()
 	GetCharacterMovement()->RotationRate = FRotator(0.0f, 600.0f, 0.0f);
 }
 
+void AShooterCharacter::BeginPlay()
+{
+	Super::BeginPlay();
+
+	// reset HP to max
+	CurrentHP = MaxHP;
+
+	// update the HUD
+	OnDamaged.Broadcast(1.0f);
+}
+
+void AShooterCharacter::EndPlay(EEndPlayReason::Type EndPlayReason)
+{
+	Super::EndPlay(EndPlayReason);
+
+	// clear the respawn timer
+	GetWorld()->GetTimerManager().ClearTimer(RespawnTimer);
+}
+
 void AShooterCharacter::SetupPlayerInputComponent(UInputComponent* PlayerInputComponent)
 {
 	// base class handles move, aim and jump inputs
@@ -52,20 +73,12 @@ float AShooterCharacter::TakeDamage(float Damage, struct FDamageEvent const& Dam
 	// Have we depleted HP?
 	if (CurrentHP <= 0.0f)
 	{
-		// deactivate the weapon
-		if (IsValid(CurrentWeapon))
-		{
-			CurrentWeapon->DeactivateWeapon();
-		}
-		
-
-		// reset the bullet counter UI
-		OnBulletCountUpdated.Broadcast(0, 0);
-
-		// destroy this character
-		Destroy();
+		Die();
 	}
 
+	// update the HUD
+	OnDamaged.Broadcast(FMath::Max(0.0f, CurrentHP / MaxHP));
+
 	return Damage;
 }
 
@@ -232,3 +245,39 @@ AShooterWeapon* AShooterCharacter::FindWeaponOfType(TSubclassOf<AShooterWeapon>
 	return nullptr;
 
 }
+
+void AShooterCharacter::Die()
+{
+	// deactivate the weapon
+	if (IsValid(CurrentWeapon))
+	{
+		CurrentWeapon->DeactivateWeapon();
+	}
+
+	// increment the team score
+	if (AShooterGameMode* GM = Cast<AShooterGameMode>(GetWorld()->GetAuthGameMode()))
+	{
+		GM->IncrementTeamScore(TeamByte);
+	}
+		
+	// stop character movement
+	GetCharacterMovement()->StopMovementImmediately();
+
+	// disable controls
+	DisableInput(nullptr);
+
+	// reset the bullet counter UI
+	OnBulletCountUpdated.Broadcast(0, 0);
+
+	// call the BP handler
+	BP_OnDeath();
+
+	// schedule character respawn
+	GetWorld()->GetTimerManager().SetTimer(RespawnTimer, this, &AShooterCharacter::OnRespawn, RespawnTime, false);
+}
+
+void AShooterCharacter::OnRespawn()
+{
+	// destroy the character to force the PC to respawn
+	Destroy();
+}
diff --git a/Templates/TP_FirstPerson/Source/TP_FirstPerson/Variant_Shooter/ShooterCharacter.h b/Templates/TP_FirstPerson/Source/TP_FirstPerson/Variant_Shooter/ShooterCharacter.h
--- a/Templates/TP_FirstPerson/Source/TP_FirstPerson/Variant_Shooter/ShooterCharacter.h
+++ b/Templates/TP_FirstPerson/Source/TP_FirstPerson/Variant_Shooter/ShooterCharacter.h
@@ -13,6 +13,7 @@ class UInputComponent;
 class UPawnNoiseEmitterComponent;
 
 DECLARE_DYNAMIC_MULTICAST_DELEGATE_TwoParams(FBulletCountUpdatedDelegate, int32, MagazineSize, int32, Bullets);
+DECLARE_DYNAMIC_MULTICAST_DELEGATE_OneParam(FDamagedDelegate, float, LifePercent);
 
 /**
  *  A player controllable first person shooter character
@@ -50,9 +51,12 @@ protected:
 	UPROPERTY(EditAnywhere, Category ="Aim", meta = (ClampMin = 0, ClampMax = 100000, Units = "cm"))
 	float MaxAimDistance = 10000.0f;
 
-	/** Current HP remaining to this character */
+	/** Max HP this character can have */
 	UPROPERTY(EditAnywhere, Category="Health")
-	float CurrentHP = 500.0f;
+	float MaxHP = 500.0f;
+
+	/** Current HP remaining to this character */
+	float CurrentHP = 0.0f;
 
 	/** Team ID for this character*/
 	UPROPERTY(EditAnywhere, Category="Team")
@@ -64,11 +68,19 @@ protected:
 	/** Weapon currently equipped and ready to shoot with */
 	TObjectPtr<AShooterWeapon> CurrentWeapon;
 
+	UPROPERTY(EditAnywhere, Category ="Destruction", meta = (ClampMin = 0, ClampMax = 10, Units = "s"))
+	float RespawnTime = 5.0f;
+
+	FTimerHandle RespawnTimer;
+
 public:
 
 	/** Bullet count updated delegate */
 	FBulletCountUpdatedDelegate OnBulletCountUpdated;
 
+	/** Damaged delegate */
+	FDamagedDelegate OnDamaged;
+
 public:
 
 	/** Constructor */
@@ -76,6 +88,12 @@ public:
 
 protected:
 
+	/** Gameplay initialization */
+	virtual void BeginPlay() override;
+
+	/** Gameplay cleanup */
+	virtual void EndPlay(EEndPlayReason::Type EndPlayReason) override;
+
 	/** Set up input action bindings */
 	virtual void SetupPlayerInputComponent(UInputComponent* InputComponent) override;
 
@@ -135,4 +153,14 @@ protected:
 
 	/** Returns true if the character already owns a weapon of the given class */
 	AShooterWeapon* FindWeaponOfType(TSubclassOf<AShooterWeapon> WeaponClass) const;
+
+	/** Called when this character's HP is depleted */
+	void Die();
+
+	/** Called to allow Blueprint code to react to this character's death */
+	UFUNCTION(BlueprintImplementableEvent, Category="Shooter", meta = (DisplayName = "On Death"))
+	void BP_OnDeath();
+
+	/** Called from the respawn timer to destroy this character and force the PC to respawn */
+	void OnRespawn();
 };
diff --git a/Templates/TP_FirstPerson/Source/TP_FirstPerson/Variant_Shooter/ShooterGameMode.cpp b/Templates/TP_FirstPerson/Source/TP_FirstPerson/Variant_Shooter/ShooterGameMode.cpp
--- a/Templates/TP_FirstPerson/Source/TP_FirstPerson/Variant_Shooter/ShooterGameMode.cpp
+++ b/Templates/TP_FirstPerson/Source/TP_FirstPerson/Variant_Shooter/ShooterGameMode.cpp
@@ -25,7 +25,8 @@ void AShooterGameMode::IncrementTeamScore(uint8 TeamByte)
 	}
 
 	// increment the score for the given team
-	TeamScores.Add(TeamByte, Score + 1);
+	++Score;
+	TeamScores.Add(TeamByte, Score);
 
 	// update the UI
 	ShooterUI->BP_UpdateScore(TeamByte, Score);
diff --git a/Templates/TP_FirstPerson/Source/TP_FirstPerson/Variant_Shooter/ShooterPlayerController.cpp b/Templates/TP_FirstPerson/Source/TP_FirstPerson/Variant_Shooter/ShooterPlayerController.cpp
--- a/Templates/TP_FirstPerson/Source/TP_FirstPerson/Variant_Shooter/ShooterPlayerController.cpp
+++ b/Templates/TP_FirstPerson/Source/TP_FirstPerson/Variant_Shooter/ShooterPlayerController.cpp
@@ -89,8 +89,12 @@ void AShooterPlayerController::OnPossess(APawn* InPawn)
 		// add the player tag
 		ShooterCharacter->Tags.Add(PlayerPawnTag);
 
-		// subscribe to the pawn's bullet count updated delegate
+		// subscribe to the pawn's delegates
 		ShooterCharacter->OnBulletCountUpdated.AddDynamic(this, &AShooterPlayerController::OnBulletCountUpdated);
+		ShooterCharacter->OnDamaged.AddDynamic(this, &AShooterPlayerController::OnPawnDamaged);
+
+		// force update the life bar
+		ShooterCharacter->OnDamaged.Broadcast(1.0f);
 	}
 }
 
@@ -124,3 +128,11 @@ void AShooterPlayerController::OnBulletCountUpdated(int32 MagazineSize, int32 Bu
 		BulletCounterUI->BP_UpdateBulletCounter(MagazineSize, Bullets);
 	}
 }
+
+void AShooterPlayerController::OnPawnDamaged(float LifePercent)
+{
+	if (IsValid(BulletCounterUI))
+	{
+		BulletCounterUI->BP_Damaged(LifePercent);
+	}
+}
diff --git a/Templates/TP_FirstPerson/Source/TP_FirstPerson/Variant_Shooter/ShooterPlayerController.h b/Templates/TP_FirstPerson/Source/TP_FirstPerson/Variant_Shooter/ShooterPlayerController.h
--- a/Templates/TP_FirstPerson/Source/TP_FirstPerson/Variant_Shooter/ShooterPlayerController.h
+++ b/Templates/TP_FirstPerson/Source/TP_FirstPerson/Variant_Shooter/ShooterPlayerController.h
@@ -74,4 +74,8 @@ protected:
 	/** Called when the bullet count on the possessed pawn is updated */
 	UFUNCTION()
 	void OnBulletCountUpdated(int32 MagazineSize, int32 Bullets);
+
+	/** Called when the possessed pawn is damaged */
+	UFUNCTION()
+	void OnPawnDamaged(float LifePercent);
 };
diff --git a/Templates/TP_FirstPerson/Source/TP_FirstPerson/Variant_Shooter/UI/ShooterBulletCounterUI.h b/Templates/TP_FirstPerson/Source/TP_FirstPerson/Variant_Shooter/UI/ShooterBulletCounterUI.h
--- a/Templates/TP_FirstPerson/Source/TP_FirstPerson/Variant_Shooter/UI/ShooterBulletCounterUI.h
+++ b/Templates/TP_FirstPerson/Source/TP_FirstPerson/Variant_Shooter/UI/ShooterBulletCounterUI.h
@@ -19,4 +19,8 @@ public:
 	/** Allows Blueprint to update sub-widgets with the new bullet count */
 	UFUNCTION(BlueprintImplementableEvent, Category="Shooter", meta=(DisplayName = "UpdateBulletCounter"))
 	void BP_UpdateBulletCounter(int32 MagazineSize, int32 BulletCount);
+
+	/** Allows Blueprint to update sub-widgets with the new life total and play a damage effect on the HUD */
+	UFUNCTION(BlueprintImplementableEvent, Category="Shooter", meta=(DisplayName = "Damaged"))
+	void BP_Damaged(float LifePercent);
 };
diff --git a/Templates/TP_FirstPerson/Source/TP_FirstPerson/Variant_Shooter/Weapons/ShooterProjectile.cpp b/Templates/TP_FirstPerson/Source/TP_FirstPerson/Variant_Shooter/Weapons/ShooterProjectile.cpp
--- a/Templates/TP_FirstPerson/Source/TP_FirstPerson/Variant_Shooter/Weapons/ShooterProjectile.cpp
+++ b/Templates/TP_FirstPerson/Source/TP_FirstPerson/Variant_Shooter/Weapons/ShooterProjectile.cpp
@@ -9,6 +9,9 @@
 #include "GameFramework/DamageType.h"
 #include "GameFramework/Pawn.h"
 #include "GameFramework/Controller.h"
+#include "Engine/OverlapResult.h"
+#include "Engine/World.h"
+#include "TimerManager.h"
 
 AShooterProjectile::AShooterProjectile()
 {
@@ -41,6 +44,14 @@ void AShooterProjectile::BeginPlay()
 	CollisionComponent->IgnoreActorWhenMoving(GetInstigator(), true);
 }
 
+void AShooterProjectile::EndPlay(EEndPlayReason::Type EndPlayReason)
+{
+	Super::EndPlay(EndPlayReason);
+
+	// clear the destruction timer
+	GetWorld()->GetTimerManager().ClearTimer(DestructionTimer);
+}
+
 void AShooterProjectile::NotifyHit(class UPrimitiveComponent* MyComp, AActor* Other, class UPrimitiveComponent* OtherComp, bool bSelfMoved, FVector HitLocation, FVector HitNormal, FVector NormalImpulse, const FHitResult& Hit)
 {
 	// ignore if we've already hit something else
@@ -51,36 +62,106 @@ void AShooterProjectile::NotifyHit(class UPrimitiveComponent* MyComp, AActor* Ot
 
 	bHit = true;
 
+	// disable collision on the projectile
+	CollisionComponent->SetCollisionEnabled(ECollisionEnabled::NoCollision);
+
 	// make AI perception noise
 	MakeNoise(NoiseLoudness, GetInstigator(), GetActorLocation(), NoiseRange, NoiseTag);
 
-	// have we hit a physics object?
-	if (OtherComp->IsSimulatingPhysics())
+	if (bExplodeOnHit)
 	{
-		// give some physics impulse to the object
-		OtherComp->AddImpulseAtLocation(GetVelocity() * PhysicsForce, Hit.ImpactPoint);
+		
+		// apply explosion damage centered on the projectile
+		ExplosionCheck(GetActorLocation());
+
+	} else {
+
+		// single hit projectile. Process the collided actor
+		ProcessHit(Other, OtherComp, Hit.ImpactPoint, -Hit.ImpactNormal);
+
+	}
+
+	// pass control to BP for any extra effects
+	BP_OnProjectileHit(Hit);
+
+	// check if we should schedule deferred destruction of the projectile
+	if (DeferredDestructionTime > 0.0f)
+	{
+		GetWorld()->GetTimerManager().SetTimer(DestructionTimer, this, &AShooterProjectile::OnDeferredDestruction, DeferredDestructionTime, false);
+
+	} else {
+
+		// destroy the projectile right away
+		Destroy();
+	}
+}
+
+void AShooterProjectile::ExplosionCheck(const FVector& ExplosionCenter)
+{
+	// do a sphere overlap check look for nearby actors to damage
+	TArray<FOverlapResult> Overlaps;
+
+	FCollisionShape OverlapShape;
+	OverlapShape.SetSphere(ExplosionRadius);
+
+	FCollisionObjectQueryParams ObjectParams;
+	ObjectParams.AddObjectTypesToQuery(ECC_Pawn);
+	ObjectParams.AddObjectTypesToQuery(ECC_WorldDynamic);
+	ObjectParams.AddObjectTypesToQuery(ECC_PhysicsBody);
+
+	FCollisionQueryParams QueryParams;
+	QueryParams.AddIgnoredActor(this);
+	if (!bDamageOwner)
+	{
+		QueryParams.AddIgnoredActor(GetInstigator());
+	}
+
+	GetWorld()->OverlapMultiByObjectType(Overlaps, ExplosionCenter, FQuat::Identity, ObjectParams, OverlapShape, QueryParams);
+
+	TArray<AActor*> DamagedActors;
+
+	// process the overlap results
+	for (const FOverlapResult& CurrentOverlap : Overlaps)
+	{
+		// overlaps may return the same actor multiple times per each component overlapped
+		// ensure we only damage each actor once by adding it to a damaged list
+		if (DamagedActors.Find(CurrentOverlap.GetActor()) == INDEX_NONE)
+		{
+			DamagedActors.Add(CurrentOverlap.GetActor());
+
+			// apply physics force away from the explosion
+			const FVector& ExplosionDir = CurrentOverlap.GetActor()->GetActorLocation() - GetActorLocation();
+
+			// push and/or damage the overlapped actor
+			ProcessHit(CurrentOverlap.GetActor(), CurrentOverlap.GetComponent(), GetActorLocation(), ExplosionDir.GetSafeNormal());
+		}
+			
 	}
+}
 
+void AShooterProjectile::ProcessHit(AActor* HitActor, UPrimitiveComponent* HitComp, const FVector& HitLocation, const FVector& HitDirection)
+{
 	// have we hit a character?
-	if (ACharacter* HitCharacter = Cast<ACharacter>(Other))
+	if (ACharacter* HitCharacter = Cast<ACharacter>(HitActor))
 	{
 		// ignore the owner of this projectile
-		if (HitCharacter != GetOwner())
+		if (HitCharacter != GetOwner() || bDamageOwner)
 		{
-			// damage the hit character
-			DamageCharacter(HitCharacter, Hit);
+			// apply damage to the character
+			UGameplayStatics::ApplyDamage(HitCharacter, HitDamage, GetInstigator()->GetController(), this, HitDamageType);
 		}
 	}
 
-	// disable collision on the projectile
-	CollisionComponent->SetCollisionEnabled(ECollisionEnabled::NoCollision);
-
-	// pass control to BP for any extra effects
-	BP_OnProjectileHit(Hit);
+	// have we hit a physics object?
+	if (HitComp->IsSimulatingPhysics())
+	{
+		// give some physics impulse to the object
+		HitComp->AddImpulseAtLocation(HitDirection * PhysicsForce, HitLocation);
+	}
 }
 
-void AShooterProjectile::DamageCharacter(ACharacter* HitCharacter, const FHitResult& Hit)
+void AShooterProjectile::OnDeferredDestruction()
 {
-	// apply damage to the character
-	UGameplayStatics::ApplyDamage(HitCharacter, HitDamage, GetInstigator()->GetController(), this, HitDamageType);
+	// destroy this actor
+	Destroy();
 }
diff --git a/Templates/TP_FirstPerson/Source/TP_FirstPerson/Variant_Shooter/Weapons/ShooterProjectile.h b/Templates/TP_FirstPerson/Source/TP_FirstPerson/Variant_Shooter/Weapons/ShooterProjectile.h
--- a/Templates/TP_FirstPerson/Source/TP_FirstPerson/Variant_Shooter/Weapons/ShooterProjectile.h
+++ b/Templates/TP_FirstPerson/Source/TP_FirstPerson/Variant_Shooter/Weapons/ShooterProjectile.h
@@ -9,6 +9,7 @@
 class USphereComponent;
 class UProjectileMovementComponent;
 class ACharacter;
+class UPrimitiveComponent;
 
 /**
  *  Simple projectile class for a first person shooter game
@@ -29,11 +30,11 @@ class TP_FIRSTPERSON_API AShooterProjectile : public AActor
 protected:
 
 	/** Loudness of the AI perception noise done by this projectile on hit */
-	UPROPERTY(EditAnywhere, Category="Noise", meta = (ClampMin = 0, ClampMax = 100))
+	UPROPERTY(EditAnywhere, Category="Projectile|Noise", meta = (ClampMin = 0, ClampMax = 100))
 	float NoiseLoudness = 3.0f;
 
 	/** Range of the AI perception noise done by this projectile on hit */
-	UPROPERTY(EditAnywhere, Category="Noise", meta = (ClampMin = 0, ClampMax = 100000, Units = "cm"))
+	UPROPERTY(EditAnywhere, Category="Projectile|Noise", meta = (ClampMin = 0, ClampMax = 100000, Units = "cm"))
 	float NoiseRange = 3000.0f;
 
 	/** Tag of the AI perception noise done by this projectile on hit */
@@ -41,20 +42,39 @@ protected:
 	FName NoiseTag = FName("Projectile");
 
 	/** Physics force to apply on hit */
-	UPROPERTY(EditAnywhere, Category="Hit", meta = (ClampMin = 0, ClampMax = 1000))
+	UPROPERTY(EditAnywhere, Category="Projectile|Hit", meta = (ClampMin = 0, ClampMax = 50000))
 	float PhysicsForce = 100.0f;
 
 	/** Damage to apply on hit */
-	UPROPERTY(EditAnywhere, Category="Hit", meta = (ClampMin = 0, ClampMax = 100))
+	UPROPERTY(EditAnywhere, Category="Projectile|Hit", meta = (ClampMin = 0, ClampMax = 100))
 	float HitDamage = 25.0f;
 
 	/** Type of damage to apply. Can be used to represent specific types of damage such as fire, explosion, etc. */
-	UPROPERTY(EditAnywhere, Category="Hit")
+	UPROPERTY(EditAnywhere, Category="Projectile|Hit")
 	TSubclassOf<UDamageType> HitDamageType;
 
+	/** If true, the projectile can damage the character that shot it */
+	UPROPERTY(EditAnywhere, Category="Projectile|Hit")
+	bool bDamageOwner = false;
+
+	/** If true, the projectile will explode and apply radial damage to all actors in range */
+	UPROPERTY(EditAnywhere, Category="Projectile|Explosion")
+	bool bExplodeOnHit = false;
+
+	/** Max distance for actors to be affected by explosion damage */
+	UPROPERTY(EditAnywhere, Category="Projectile|Explosion", meta = (ClampMin = 0, ClampMax = 5000, Units = "cm"))
+	float ExplosionRadius = 500.0f;	
+
 	/** If true, this projectile has already hit another surface */
 	bool bHit = false;
 
+	/** How long to wait after a hit before destroying this projectile */
+	UPROPERTY(EditAnywhere, Category="Projectile|Destruction", meta = (ClampMin = 0, ClampMax = 10, Units = "s"))
+	float DeferredDestructionTime = 5.0f;
+
+	/** Timer to handle deferred destruction of this projectile */
+	FTimerHandle DestructionTimer;
+
 public:	
 
 	/** Constructor */
@@ -65,16 +85,25 @@ protected:
 	/** Gameplay initialization */
 	virtual void BeginPlay() override;
 
+	/** Gameplay cleanup */
+	virtual void EndPlay(EEndPlayReason::Type EndPlayReason) override;
+
 	/** Handles collision */
-	virtual void NotifyHit(class UPrimitiveComponent* MyComp, AActor* Other, class UPrimitiveComponent* OtherComp, bool bSelfMoved, FVector HitLocation, FVector HitNormal, FVector NormalImpulse, const FHitResult& Hit) override;
+	virtual void NotifyHit(class UPrimitiveComponent* MyComp, AActor* Other, UPrimitiveComponent* OtherComp, bool bSelfMoved, FVector HitLocation, FVector HitNormal, FVector NormalImpulse, const FHitResult& Hit) override;
 
 protected:
 
-	/** Apply damage to a hit character */
-	UFUNCTION(BlueprintCallable, Category="Projectile")
-	virtual void DamageCharacter(ACharacter* HitCharacter, const FHitResult& Hit);
+	/** Looks up actors within the explosion radius and damages them */
+	void ExplosionCheck(const FVector& ExplosionCenter);
+
+	/** Processes a projectile hit for the given actor */
+	void ProcessHit(AActor* HitActor, UPrimitiveComponent* HitComp, const FVector& HitLocation, const FVector& HitDirection);
 
-	/** Passes control to Blueprint to implement any effects on hit */
+	/** Passes control to Blueprint to implement any effects on hit. */
 	UFUNCTION(BlueprintImplementableEvent, Category="Projectile", meta = (DisplayName = "On Projectile Hit"))
 	void BP_OnProjectileHit(const FHitResult& Hit);
+
+	/** Called from the destruction timer to destroy this projectile */
+	void OnDeferredDestruction();
+
 };
diff --git a/Templates/TP_FirstPersonBP/Config/TemplateDefs.ini b/Templates/TP_FirstPersonBP/Config/TemplateDefs.ini
--- a/Templates/TP_FirstPersonBP/Config/TemplateDefs.ini
+++ b/Templates/TP_FirstPersonBP/Config/TemplateDefs.ini
@@ -39,9 +39,9 @@ EditDetailLevelPreference="High"
 
 ; Variant packs
 ; NOTE: make sure these are commented out for release!
-;SharedContentPacks=(MountName="Variant_Shooter",DetailLevels=("Standard"))
-;SharedContentPacks=(MountName="Variant_Horror",DetailLevels=("Standard"))
-;SharedContentPacks=(MountName="Weapons",DetailLevels=("Standard"))
+SharedContentPacks=(MountName="Variant_Shooter",DetailLevels=("Standard"))
+SharedContentPacks=(MountName="Variant_Horror",DetailLevels=("Standard"))
+SharedContentPacks=(MountName="Weapons",DetailLevels=("Standard"))
 
 ; Variant definitions
 Variants=(Name="ArenaShooter",LocalizedDisplayNames=((Language="en",Text="Arena Shooter")),LocalizedDescriptions=((Language="en",Text="An arena shooter with weapon pickups and an enemy AI opponent.")),SharedContentPacks=((DetailLevels=(Standard),MountName="Weapons"),(DetailLevels=(Standard),MountName="Variant_Shooter")) )
diff --git a/Templates/TP_VehicleAdvBP/Config/TemplateDefs.ini b/Templates/TP_VehicleAdvBP/Config/TemplateDefs.ini
--- a/Templates/TP_VehicleAdvBP/Config/TemplateDefs.ini
+++ b/Templates/TP_VehicleAdvBP/Config/TemplateDefs.ini
@@ -42,8 +42,8 @@ EditDetailLevelPreference="High"
 
 ; Variant packs
 ; NOTE: make sure these are commented out for release!
-SharedContentPacks=(MountName="Variant_Offroad",DetailLevels=("Standard"))
-SharedContentPacks=(MountName="Variant_Timetrial",DetailLevels=("Standard"))
+;SharedContentPacks=(MountName="Variant_Offroad",DetailLevels=("Standard"))
+;SharedContentPacks=(MountName="Variant_Timetrial",DetailLevels=("Standard"))
 
 Variants=(Name="Timetrial",LocalizedDisplayNames=((Language="en",Text="Timetrial")),LocalizedDescriptions=((Language="en",Text="A timed lap with UI and track markers.")),SharedContentPacks=((DetailLevels=(Standard),MountName="Variant_Timetrial")))
 
