Templates
- Removed manual bUseMobileControls flag
- Implemented automatic platform detection on BP templates, spawns mobile controls on iOS and Android
- Implemented automatic touchscreen use detection on CPP templates, spawns mobile controls for touchscreen platforms.
- TP_FirstPerson, BP: Fixed corner case where bullet counter HUD would not hide when the player dies by KillZ

#rnx

diff --git a/Templates/TP_FirstPerson/Source/TP_FirstPerson/TP_FirstPerson.Build.cs b/Templates/TP_FirstPerson/Source/TP_FirstPerson/TP_FirstPerson.Build.cs
--- a/Templates/TP_FirstPerson/Source/TP_FirstPerson/TP_FirstPerson.Build.cs
+++ b/Templates/TP_FirstPerson/Source/TP_FirstPerson/TP_FirstPerson.Build.cs
@@ -17,7 +17,8 @@ public class TP_FirstPerson : ModuleRules
 			"AIModule",
 			"StateTreeModule",
 			"GameplayStateTreeModule",
-			"UMG"
+			"UMG",
+			"Slate"
 		});
 
 		PrivateDependencyModuleNames.AddRange(new string[] { });
diff --git a/Templates/TP_FirstPerson/Source/TP_FirstPerson/TP_FirstPersonPlayerController.cpp b/Templates/TP_FirstPerson/Source/TP_FirstPerson/TP_FirstPersonPlayerController.cpp
--- a/Templates/TP_FirstPerson/Source/TP_FirstPerson/TP_FirstPersonPlayerController.cpp
+++ b/Templates/TP_FirstPerson/Source/TP_FirstPerson/TP_FirstPersonPlayerController.cpp
@@ -8,6 +8,7 @@
 #include "TP_FirstPersonCameraManager.h"
 #include "Blueprint/UserWidget.h"
 #include "TP_FirstPerson.h"
+#include "Widgets/Input/SVirtualJoystick.h"
 
 ATP_FirstPersonPlayerController::ATP_FirstPersonPlayerController()
 {
@@ -19,8 +20,9 @@ void ATP_FirstPersonPlayerController::BeginPlay()
 {
 	Super::BeginPlay();
 
+	
 	// only spawn touch controls on local player controllers
-	if (bUseMobileControlsWidget && IsLocalPlayerController())
+	if (SVirtualJoystick::ShouldDisplayTouchInterface() && IsLocalPlayerController())
 	{
 		// spawn the mobile controls widget
 		MobileControlsWidget = CreateWidget<UUserWidget>(this, MobileControlsWidgetClass);
@@ -55,7 +57,7 @@ void ATP_FirstPersonPlayerController::SetupInputComponent()
 			}
 
 			// only add these IMCs if we're not using mobile touch input
-			if (!bUseMobileControlsWidget)
+			if (!SVirtualJoystick::ShouldDisplayTouchInterface())
 			{
 				for (UInputMappingContext* CurrentContext : MobileExcludedMappingContexts)
 				{
diff --git a/Templates/TP_FirstPerson/Source/TP_FirstPerson/TP_FirstPersonPlayerController.h b/Templates/TP_FirstPerson/Source/TP_FirstPerson/TP_FirstPersonPlayerController.h
--- a/Templates/TP_FirstPerson/Source/TP_FirstPerson/TP_FirstPersonPlayerController.h
+++ b/Templates/TP_FirstPerson/Source/TP_FirstPerson/TP_FirstPersonPlayerController.h
@@ -34,12 +34,8 @@ protected:
 	UPROPERTY(EditAnywhere, Category="Input|Input Mappings")
 	TArray<UInputMappingContext*> MobileExcludedMappingContexts;
 
-	/** If true, a mobile controls widget will be used for touch input */
-	UPROPERTY(EditAnywhere, Category="Input|Touch Controls")
-	bool bUseMobileControlsWidget = false;
-
 	/** Mobile controls widget to spawn */
-	UPROPERTY(EditAnywhere, Category="Input|Touch Controls", meta = (EditCondition = "bUseMobileControlsWidget"))
+	UPROPERTY(EditAnywhere, Category="Input|Touch Controls")
 	TSubclassOf<UUserWidget> MobileControlsWidgetClass;
 
 	/** Pointer to the mobile controls widget */
diff --git a/Templates/TP_FirstPerson/Source/TP_FirstPerson/Variant_Horror/HorrorPlayerController.cpp b/Templates/TP_FirstPerson/Source/TP_FirstPerson/Variant_Horror/HorrorPlayerController.cpp
--- a/Templates/TP_FirstPerson/Source/TP_FirstPerson/Variant_Horror/HorrorPlayerController.cpp
+++ b/Templates/TP_FirstPerson/Source/TP_FirstPerson/Variant_Horror/HorrorPlayerController.cpp
@@ -8,7 +8,8 @@
 #include "TP_FirstPersonCameraManager.h"
 #include "HorrorCharacter.h"
 #include "HorrorUI.h"
-#include "../TP_FirstPerson.h"
+#include "TP_FirstPerson.h"
+#include "Widgets/Input/SVirtualJoystick.h"
 
 AHorrorPlayerController::AHorrorPlayerController()
 {
@@ -21,7 +22,7 @@ void AHorrorPlayerController::BeginPlay()
 	Super::BeginPlay();
 
 	// only spawn touch controls on local player controllers
-	if (bUseMobileControlsWidget && IsLocalPlayerController())
+	if (SVirtualJoystick::ShouldDisplayTouchInterface() && IsLocalPlayerController())
 	{
 		// spawn the mobile controls widget
 		MobileControlsWidget = CreateWidget<UUserWidget>(this, MobileControlsWidgetClass);
@@ -79,7 +80,7 @@ void AHorrorPlayerController::SetupInputComponent()
 			}
 
 			// only add these IMCs if we're not using mobile touch input
-			if (!bUseMobileControlsWidget)
+			if (!SVirtualJoystick::ShouldDisplayTouchInterface())
 			{
 				for (UInputMappingContext* CurrentContext : MobileExcludedMappingContexts)
 				{
diff --git a/Templates/TP_FirstPerson/Source/TP_FirstPerson/Variant_Horror/HorrorPlayerController.h b/Templates/TP_FirstPerson/Source/TP_FirstPerson/Variant_Horror/HorrorPlayerController.h
--- a/Templates/TP_FirstPerson/Source/TP_FirstPerson/Variant_Horror/HorrorPlayerController.h
+++ b/Templates/TP_FirstPerson/Source/TP_FirstPerson/Variant_Horror/HorrorPlayerController.h
@@ -43,12 +43,8 @@ protected:
 	UPROPERTY(EditAnywhere, Category="Input|Input Mappings")
 	TArray<UInputMappingContext*> MobileExcludedMappingContexts;
 
-	/** If true, a mobile controls widget will be used for touch input */
-	UPROPERTY(EditAnywhere, Category="Input|Touch Controls")
-	bool bUseMobileControlsWidget = false;
-
 	/** Mobile controls widget to spawn */
-	UPROPERTY(EditAnywhere, Category="Input|Touch Controls", meta = (EditCondition = "bUseMobileControlsWidget"))
+	UPROPERTY(EditAnywhere, Category="Input|Touch Controls")
 	TSubclassOf<UUserWidget> MobileControlsWidgetClass;
 
 	/** Pointer to the mobile controls widget */
diff --git a/Templates/TP_FirstPerson/Source/TP_FirstPerson/Variant_Shooter/ShooterPlayerController.cpp b/Templates/TP_FirstPerson/Source/TP_FirstPerson/Variant_Shooter/ShooterPlayerController.cpp
--- a/Templates/TP_FirstPerson/Source/TP_FirstPerson/Variant_Shooter/ShooterPlayerController.cpp
+++ b/Templates/TP_FirstPerson/Source/TP_FirstPerson/Variant_Shooter/ShooterPlayerController.cpp
@@ -10,6 +10,7 @@
 #include "ShooterCharacter.h"
 #include "ShooterBulletCounterUI.h"
 #include "TP_FirstPerson.h"
+#include "Widgets/Input/SVirtualJoystick.h"
 
 void AShooterPlayerController::BeginPlay()
 {
@@ -18,7 +19,7 @@ void AShooterPlayerController::BeginPlay()
 	// only spawn touch controls on local player controllers
 	if (IsLocalPlayerController())
 	{
-		if (bUseMobileControlsWidget)
+		if (SVirtualJoystick::ShouldDisplayTouchInterface())
 		{
 			// spawn the mobile controls widget
 			MobileControlsWidget = CreateWidget<UUserWidget>(this, MobileControlsWidgetClass);
@@ -65,7 +66,7 @@ void AShooterPlayerController::SetupInputComponent()
 			}
 
 			// only add these IMCs if we're not using mobile touch input
-			if (!bUseMobileControlsWidget)
+			if (!SVirtualJoystick::ShouldDisplayTouchInterface())
 			{
 				for (UInputMappingContext* CurrentContext : MobileExcludedMappingContexts)
 				{
@@ -100,6 +101,9 @@ void AShooterPlayerController::OnPossess(APawn* InPawn)
 
 void AShooterPlayerController::OnPawnDestroyed(AActor* DestroyedActor)
 {
+	// reset the bullet counter HUD
+	BulletCounterUI->BP_UpdateBulletCounter(0, 0);
+
 	// find the player start
 	TArray<AActor*> ActorList;
 	UGameplayStatics::GetAllActorsOfClass(GetWorld(), APlayerStart::StaticClass(), ActorList);
diff --git a/Templates/TP_FirstPerson/Source/TP_FirstPerson/Variant_Shooter/ShooterPlayerController.h b/Templates/TP_FirstPerson/Source/TP_FirstPerson/Variant_Shooter/ShooterPlayerController.h
--- a/Templates/TP_FirstPerson/Source/TP_FirstPerson/Variant_Shooter/ShooterPlayerController.h
+++ b/Templates/TP_FirstPerson/Source/TP_FirstPerson/Variant_Shooter/ShooterPlayerController.h
@@ -30,12 +30,8 @@ protected:
 	UPROPERTY(EditAnywhere, Category="Input|Input Mappings")
 	TArray<UInputMappingContext*> MobileExcludedMappingContexts;
 
-	/** If true, a mobile controls widget will be used for touch input */
-	UPROPERTY(EditAnywhere, Category="Input|Touch Controls")
-	bool bUseMobileControlsWidget = false;
-
 	/** Mobile controls widget to spawn */
-	UPROPERTY(EditAnywhere, Category="Input|Touch Controls", meta = (EditCondition = "bUseMobileControlsWidget"))
+	UPROPERTY(EditAnywhere, Category="Input|Touch Controls")
 	TSubclassOf<UUserWidget> MobileControlsWidgetClass;
 
 	/** Pointer to the mobile controls widget */
diff --git a/Templates/TP_ThirdPerson/Source/TP_ThirdPerson/TP_ThirdPerson.Build.cs b/Templates/TP_ThirdPerson/Source/TP_ThirdPerson/TP_ThirdPerson.Build.cs
--- a/Templates/TP_ThirdPerson/Source/TP_ThirdPerson/TP_ThirdPerson.Build.cs
+++ b/Templates/TP_ThirdPerson/Source/TP_ThirdPerson/TP_ThirdPerson.Build.cs
@@ -17,7 +17,8 @@ public class TP_ThirdPerson : ModuleRules
 			"AIModule",
 			"StateTreeModule",
 			"GameplayStateTreeModule",
-			"UMG"
+			"UMG",
+			"Slate"
 		});
 
 		PrivateDependencyModuleNames.AddRange(new string[] { });
diff --git a/Templates/TP_ThirdPerson/Source/TP_ThirdPerson/TP_ThirdPersonPlayerController.cpp b/Templates/TP_ThirdPerson/Source/TP_ThirdPerson/TP_ThirdPersonPlayerController.cpp
--- a/Templates/TP_ThirdPerson/Source/TP_ThirdPerson/TP_ThirdPersonPlayerController.cpp
+++ b/Templates/TP_ThirdPerson/Source/TP_ThirdPerson/TP_ThirdPersonPlayerController.cpp
@@ -7,13 +7,14 @@
 #include "InputMappingContext.h"
 #include "Blueprint/UserWidget.h"
 #include "TP_ThirdPerson.h"
+#include "Widgets/Input/SVirtualJoystick.h"
 
 void ATP_ThirdPersonPlayerController::BeginPlay()
 {
 	Super::BeginPlay();
 
 	// only spawn touch controls on local player controllers
-	if (bUseMobileControlsWidget && IsLocalPlayerController())
+	if (SVirtualJoystick::ShouldDisplayTouchInterface() && IsLocalPlayerController())
 	{
 		// spawn the mobile controls widget
 		MobileControlsWidget = CreateWidget<UUserWidget>(this, MobileControlsWidgetClass);
@@ -48,7 +49,7 @@ void ATP_ThirdPersonPlayerController::SetupInputComponent()
 			}
 
 			// only add these IMCs if we're not using mobile touch input
-			if (!bUseMobileControlsWidget)
+			if (!SVirtualJoystick::ShouldDisplayTouchInterface())
 			{
 				for (UInputMappingContext* CurrentContext : MobileExcludedMappingContexts)
 				{
diff --git a/Templates/TP_ThirdPerson/Source/TP_ThirdPerson/TP_ThirdPersonPlayerController.h b/Templates/TP_ThirdPerson/Source/TP_ThirdPerson/TP_ThirdPersonPlayerController.h
--- a/Templates/TP_ThirdPerson/Source/TP_ThirdPerson/TP_ThirdPersonPlayerController.h
+++ b/Templates/TP_ThirdPerson/Source/TP_ThirdPerson/TP_ThirdPersonPlayerController.h
@@ -28,12 +28,8 @@ protected:
 	UPROPERTY(EditAnywhere, Category="Input|Input Mappings")
 	TArray<UInputMappingContext*> MobileExcludedMappingContexts;
 
-	/** If true, a mobile controls widget will be used for touch input */
-	UPROPERTY(EditAnywhere, Category="Input|Touch Controls")
-	bool bUseMobileControlsWidget = false;
-
 	/** Mobile controls widget to spawn */
-	UPROPERTY(EditAnywhere, Category="Input|Touch Controls", meta = (EditCondition = "bUseMobileControlsWidget"))
+	UPROPERTY(EditAnywhere, Category="Input|Touch Controls")
 	TSubclassOf<UUserWidget> MobileControlsWidgetClass;
 
 	/** Pointer to the mobile controls widget */
diff --git a/Templates/TP_ThirdPerson/Source/TP_ThirdPerson/Variant_Combat/CombatPlayerController.cpp b/Templates/TP_ThirdPerson/Source/TP_ThirdPerson/Variant_Combat/CombatPlayerController.cpp
--- a/Templates/TP_ThirdPerson/Source/TP_ThirdPerson/Variant_Combat/CombatPlayerController.cpp
+++ b/Templates/TP_ThirdPerson/Source/TP_ThirdPerson/Variant_Combat/CombatPlayerController.cpp
@@ -11,13 +11,14 @@
 #include "Engine/World.h"
 #include "Blueprint/UserWidget.h"
 #include "TP_ThirdPerson.h"
+#include "Widgets/Input/SVirtualJoystick.h"
 
 void ACombatPlayerController::BeginPlay()
 {
 	Super::BeginPlay();
 
 	// only spawn touch controls on local player controllers
-	if (bUseMobileControlsWidget && IsLocalPlayerController())
+	if (SVirtualJoystick::ShouldDisplayTouchInterface() && IsLocalPlayerController())
 	{
 		// spawn the mobile controls widget
 		MobileControlsWidget = CreateWidget<UUserWidget>(this, MobileControlsWidgetClass);
@@ -50,7 +51,7 @@ void ACombatPlayerController::SetupInputComponent()
 			}
 
 			// only add these IMCs if we're not using mobile touch input
-			if (!bUseMobileControlsWidget)
+			if (!SVirtualJoystick::ShouldDisplayTouchInterface())
 			{
 				for (UInputMappingContext* CurrentContext : MobileExcludedMappingContexts)
 				{
diff --git a/Templates/TP_ThirdPerson/Source/TP_ThirdPerson/Variant_Combat/CombatPlayerController.h b/Templates/TP_ThirdPerson/Source/TP_ThirdPerson/Variant_Combat/CombatPlayerController.h
--- a/Templates/TP_ThirdPerson/Source/TP_ThirdPerson/Variant_Combat/CombatPlayerController.h
+++ b/Templates/TP_ThirdPerson/Source/TP_ThirdPerson/Variant_Combat/CombatPlayerController.h
@@ -29,12 +29,8 @@ protected:
 	UPROPERTY(EditAnywhere, Category="Input|Input Mappings")
 	TArray<UInputMappingContext*> MobileExcludedMappingContexts;
 
-	/** If true, a mobile controls widget will be used for touch input */
-	UPROPERTY(EditAnywhere, Category="Input|Touch Controls")
-	bool bUseMobileControlsWidget = false;
-
 	/** Mobile controls widget to spawn */
-	UPROPERTY(EditAnywhere, Category="Input|Touch Controls", meta = (EditCondition = "bUseMobileControlsWidget"))
+	UPROPERTY(EditAnywhere, Category="Input|Touch Controls")
 	TSubclassOf<UUserWidget> MobileControlsWidgetClass;
 
 	/** Pointer to the mobile controls widget */
diff --git a/Templates/TP_ThirdPerson/Source/TP_ThirdPerson/Variant_Platforming/PlatformingPlayerController.cpp b/Templates/TP_ThirdPerson/Source/TP_ThirdPerson/Variant_Platforming/PlatformingPlayerController.cpp
--- a/Templates/TP_ThirdPerson/Source/TP_ThirdPerson/Variant_Platforming/PlatformingPlayerController.cpp
+++ b/Templates/TP_ThirdPerson/Source/TP_ThirdPerson/Variant_Platforming/PlatformingPlayerController.cpp
@@ -11,13 +11,14 @@
 #include "Engine/World.h"
 #include "Blueprint/UserWidget.h"
 #include "TP_ThirdPerson.h"
+#include "Widgets/Input/SVirtualJoystick.h"
 
 void APlatformingPlayerController::BeginPlay()
 {
 	Super::BeginPlay();
 
 	// only spawn touch controls on local player controllers
-	if (bUseMobileControlsWidget && IsLocalPlayerController())
+	if (SVirtualJoystick::ShouldDisplayTouchInterface() && IsLocalPlayerController())
 	{
 		// spawn the mobile controls widget
 		MobileControlsWidget = CreateWidget<UUserWidget>(this, MobileControlsWidgetClass);
@@ -50,7 +51,7 @@ void APlatformingPlayerController::SetupInputComponent()
 			}
 
 			// only add these IMCs if we're not using mobile touch input
-			if (!bUseMobileControlsWidget)
+			if (!SVirtualJoystick::ShouldDisplayTouchInterface())
 			{
 				for (UInputMappingContext* CurrentContext : MobileExcludedMappingContexts)
 				{
diff --git a/Templates/TP_ThirdPerson/Source/TP_ThirdPerson/Variant_Platforming/PlatformingPlayerController.h b/Templates/TP_ThirdPerson/Source/TP_ThirdPerson/Variant_Platforming/PlatformingPlayerController.h
--- a/Templates/TP_ThirdPerson/Source/TP_ThirdPerson/Variant_Platforming/PlatformingPlayerController.h
+++ b/Templates/TP_ThirdPerson/Source/TP_ThirdPerson/Variant_Platforming/PlatformingPlayerController.h
@@ -29,12 +29,8 @@ protected:
 	UPROPERTY(EditAnywhere, Category="Input|Input Mappings")
 	TArray<UInputMappingContext*> MobileExcludedMappingContexts;
 
-	/** If true, a mobile controls widget will be used for touch input */
-	UPROPERTY(EditAnywhere, Category="Input|Touch Controls")
-	bool bUseMobileControlsWidget = false;
-
 	/** Mobile controls widget to spawn */
-	UPROPERTY(EditAnywhere, Category="Input|Touch Controls", meta = (EditCondition = "bUseMobileControlsWidget"))
+	UPROPERTY(EditAnywhere, Category="Input|Touch Controls")
 	TSubclassOf<UUserWidget> MobileControlsWidgetClass;
 
 	/** Pointer to the mobile controls widget */
diff --git a/Templates/TP_ThirdPerson/Source/TP_ThirdPerson/Variant_SideScrolling/SideScrollingPlayerController.cpp b/Templates/TP_ThirdPerson/Source/TP_ThirdPerson/Variant_SideScrolling/SideScrollingPlayerController.cpp
--- a/Templates/TP_ThirdPerson/Source/TP_ThirdPerson/Variant_SideScrolling/SideScrollingPlayerController.cpp
+++ b/Templates/TP_ThirdPerson/Source/TP_ThirdPerson/Variant_SideScrolling/SideScrollingPlayerController.cpp
@@ -11,13 +11,14 @@
 #include "Engine/World.h"
 #include "Blueprint/UserWidget.h"
 #include "TP_ThirdPerson.h"
+#include "Widgets/Input/SVirtualJoystick.h"
 
 void ASideScrollingPlayerController::BeginPlay()
 {
 	Super::BeginPlay();
 
 	// only spawn touch controls on local player controllers
-	if (bUseMobileControlsWidget && IsLocalPlayerController())
+	if (SVirtualJoystick::ShouldDisplayTouchInterface() && IsLocalPlayerController())
 	{
 		// spawn the mobile controls widget
 		MobileControlsWidget = CreateWidget<UUserWidget>(this, MobileControlsWidgetClass);
@@ -50,7 +51,7 @@ void ASideScrollingPlayerController::SetupInputComponent()
 			}
 
 			// only add these IMCs if we're not using mobile touch input
-			if (!bUseMobileControlsWidget)
+			if (!SVirtualJoystick::ShouldDisplayTouchInterface())
 			{
 				for (UInputMappingContext* CurrentContext : MobileExcludedMappingContexts)
 				{
diff --git a/Templates/TP_ThirdPerson/Source/TP_ThirdPerson/Variant_SideScrolling/SideScrollingPlayerController.h b/Templates/TP_ThirdPerson/Source/TP_ThirdPerson/Variant_SideScrolling/SideScrollingPlayerController.h
--- a/Templates/TP_ThirdPerson/Source/TP_ThirdPerson/Variant_SideScrolling/SideScrollingPlayerController.h
+++ b/Templates/TP_ThirdPerson/Source/TP_ThirdPerson/Variant_SideScrolling/SideScrollingPlayerController.h
@@ -30,12 +30,8 @@ protected:
 	UPROPERTY(EditAnywhere, Category="Input|Input Mappings")
 	TArray<UInputMappingContext*> MobileExcludedMappingContexts;
 
-	/** If true, a mobile controls widget will be used for touch input */
-	UPROPERTY(EditAnywhere, Category="Input|Touch Controls")
-	bool bUseMobileControlsWidget = false;
-
 	/** Mobile controls widget to spawn */
-	UPROPERTY(EditAnywhere, Category="Input|Touch Controls", meta = (EditCondition = "bUseMobileControlsWidget"))
+	UPROPERTY(EditAnywhere, Category="Input|Touch Controls")
 	TSubclassOf<UUserWidget> MobileControlsWidgetClass;
 
 	/** Pointer to the mobile controls widget */
diff --git a/Templates/TP_TopDown/Source/TP_TopDown/TP_TopDown.Build.cs b/Templates/TP_TopDown/Source/TP_TopDown/TP_TopDown.Build.cs
--- a/Templates/TP_TopDown/Source/TP_TopDown/TP_TopDown.Build.cs
+++ b/Templates/TP_TopDown/Source/TP_TopDown/TP_TopDown.Build.cs
@@ -19,7 +19,8 @@ public class TP_TopDown : ModuleRules
 			"StateTreeModule",
 			"GameplayStateTreeModule",
 			"Niagara",
-			"UMG"
+			"UMG",
+			"Slate"
 		});
 
 		PrivateDependencyModuleNames.AddRange(new string[] { });
diff --git a/Templates/TP_TopDown/Source/TP_TopDown/Variant_TwinStick/TwinStickPlayerController.cpp b/Templates/TP_TopDown/Source/TP_TopDown/Variant_TwinStick/TwinStickPlayerController.cpp
--- a/Templates/TP_TopDown/Source/TP_TopDown/Variant_TwinStick/TwinStickPlayerController.cpp
+++ b/Templates/TP_TopDown/Source/TP_TopDown/Variant_TwinStick/TwinStickPlayerController.cpp
@@ -11,13 +11,14 @@
 #include "Engine/World.h"
 #include "Blueprint/UserWidget.h"
 #include "TP_TopDown.h"
+#include "Widgets/Input/SVirtualJoystick.h"
 
 void ATwinStickPlayerController::BeginPlay()
 {
 	Super::BeginPlay();
 
 	// only spawn touch controls on local player controllers
-	if (bUseMobileControlsWidget && IsLocalPlayerController())
+	if (SVirtualJoystick::ShouldDisplayTouchInterface() && IsLocalPlayerController())
 	{
 		// spawn the mobile controls widget
 		MobileControlsWidget = CreateWidget<UUserWidget>(this, MobileControlsWidgetClass);
diff --git a/Templates/TP_TopDown/Source/TP_TopDown/Variant_TwinStick/TwinStickPlayerController.h b/Templates/TP_TopDown/Source/TP_TopDown/Variant_TwinStick/TwinStickPlayerController.h
--- a/Templates/TP_TopDown/Source/TP_TopDown/Variant_TwinStick/TwinStickPlayerController.h
+++ b/Templates/TP_TopDown/Source/TP_TopDown/Variant_TwinStick/TwinStickPlayerController.h
@@ -25,12 +25,8 @@ protected:
 	UPROPERTY(EditAnywhere, Category ="Input|Input Mappings")
 	TArray<UInputMappingContext*> DefaultMappingContexts;
 
-	/** If true, a mobile controls widget will be used for touch input */
-	UPROPERTY(EditAnywhere, Category="Input|Touch Controls")
-	bool bUseMobileControlsWidget = false;
-
 	/** Mobile controls widget to spawn */
-	UPROPERTY(EditAnywhere, Category="Input|Touch Controls", meta = (EditCondition = "bUseMobileControlsWidget"))
+	UPROPERTY(EditAnywhere, Category="Input|Touch Controls")
 	TSubclassOf<UUserWidget> MobileControlsWidgetClass;
 
 	/** Pointer to the mobile controls widget */
diff --git a/Templates/TP_VehicleAdv/Source/TP_VehicleAdv/TP_VehicleAdv.Build.cs b/Templates/TP_VehicleAdv/Source/TP_VehicleAdv/TP_VehicleAdv.Build.cs
--- a/Templates/TP_VehicleAdv/Source/TP_VehicleAdv/TP_VehicleAdv.Build.cs
+++ b/Templates/TP_VehicleAdv/Source/TP_VehicleAdv/TP_VehicleAdv.Build.cs
@@ -16,7 +16,8 @@ public class TP_VehicleAdv : ModuleRules
 			"EnhancedInput",
 			"ChaosVehicles",
 			"PhysicsCore",
-			"UMG"
+			"UMG",
+			"Slate"
 		});
 
 		PublicIncludePaths.AddRange(new string[] {
diff --git a/Templates/TP_VehicleAdv/Source/TP_VehicleAdv/TP_VehicleAdvPlayerController.cpp b/Templates/TP_VehicleAdv/Source/TP_VehicleAdv/TP_VehicleAdvPlayerController.cpp
--- a/Templates/TP_VehicleAdv/Source/TP_VehicleAdv/TP_VehicleAdvPlayerController.cpp
+++ b/Templates/TP_VehicleAdv/Source/TP_VehicleAdv/TP_VehicleAdvPlayerController.cpp
@@ -10,6 +10,7 @@
 #include "TP_VehicleAdv.h"
 #include "Kismet/GameplayStatics.h"
 #include "GameFramework/PlayerStart.h"
+#include "Widgets/Input/SVirtualJoystick.h"
 
 void ATP_VehicleAdvPlayerController::BeginPlay()
 {
@@ -18,7 +19,7 @@ void ATP_VehicleAdvPlayerController::BeginPlay()
 	// only spawn UI on local player controllers
 	if (IsLocalPlayerController())
 	{
-		if (bUseMobileControlsWidget)
+		if (SVirtualJoystick::ShouldDisplayTouchInterface())
 		{
 			// spawn the mobile controls widget
 			MobileControlsWidget = CreateWidget<UUserWidget>(this, MobileControlsWidgetClass);
@@ -67,7 +68,7 @@ void ATP_VehicleAdvPlayerController::SetupInputComponent()
 			}
 
 			// only add these IMCs if we're not using mobile touch input
-			if (!bUseMobileControlsWidget)
+			if (!SVirtualJoystick::ShouldDisplayTouchInterface())
 			{
 				for (UInputMappingContext* CurrentContext : MobileExcludedMappingContexts)
 				{
diff --git a/Templates/TP_VehicleAdv/Source/TP_VehicleAdv/TP_VehicleAdvPlayerController.h b/Templates/TP_VehicleAdv/Source/TP_VehicleAdv/TP_VehicleAdvPlayerController.h
--- a/Templates/TP_VehicleAdv/Source/TP_VehicleAdv/TP_VehicleAdvPlayerController.h
+++ b/Templates/TP_VehicleAdv/Source/TP_VehicleAdv/TP_VehicleAdvPlayerController.h
@@ -29,12 +29,8 @@ protected:
 	UPROPERTY(EditAnywhere, Category="Input|Input Mappings")
 	TArray<UInputMappingContext*> MobileExcludedMappingContexts;
 
-	/** If true, a mobile controls widget will be used for touch input */
-	UPROPERTY(EditAnywhere, Category="Input|Touch Controls")
-	bool bUseMobileControlsWidget = false;
-
 	/** Mobile controls widget to spawn */
-	UPROPERTY(EditAnywhere, Category="Input|Touch Controls", meta = (EditCondition = "bUseMobileControlsWidget"))
+	UPROPERTY(EditAnywhere, Category="Input|Touch Controls")
 	TSubclassOf<UUserWidget> MobileControlsWidgetClass;
 
 	/** Pointer to the mobile controls widget */
diff --git a/Templates/TP_VehicleAdv/Source/TP_VehicleAdv/Variant_TimeTrial/TimeTrialPlayerController.cpp b/Templates/TP_VehicleAdv/Source/TP_VehicleAdv/Variant_TimeTrial/TimeTrialPlayerController.cpp
--- a/Templates/TP_VehicleAdv/Source/TP_VehicleAdv/Variant_TimeTrial/TimeTrialPlayerController.cpp
+++ b/Templates/TP_VehicleAdv/Source/TP_VehicleAdv/Variant_TimeTrial/TimeTrialPlayerController.cpp
@@ -16,6 +16,7 @@
 #include "TP_VehicleAdv.h"
 #include "Kismet/GameplayStatics.h"
 #include "GameFramework/PlayerStart.h"
+#include "Widgets/Input/SVirtualJoystick.h"
 
 void ATimeTrialPlayerController::BeginPlay()
 {
@@ -24,7 +25,7 @@ void ATimeTrialPlayerController::BeginPlay()
 	// only spawn UI on local player controllers
 	if (IsLocalPlayerController())
 	{
-		if (bUseMobileControlsWidget)
+		if (SVirtualJoystick::ShouldDisplayTouchInterface())
 		{
 			// spawn the mobile controls widget
 			MobileControlsWidget = CreateWidget<UUserWidget>(this, MobileControlsWidgetClass);
@@ -90,7 +91,7 @@ void ATimeTrialPlayerController::SetupInputComponent()
 			}
 
 			// only add these IMCs if we're not using mobile touch input
-			if (!bUseMobileControlsWidget)
+			if (!SVirtualJoystick::ShouldDisplayTouchInterface())
 			{
 				for (UInputMappingContext* CurrentContext : MobileExcludedMappingContexts)
 				{
diff --git a/Templates/TP_VehicleAdv/Source/TP_VehicleAdv/Variant_TimeTrial/TimeTrialPlayerController.h b/Templates/TP_VehicleAdv/Source/TP_VehicleAdv/Variant_TimeTrial/TimeTrialPlayerController.h
--- a/Templates/TP_VehicleAdv/Source/TP_VehicleAdv/Variant_TimeTrial/TimeTrialPlayerController.h
+++ b/Templates/TP_VehicleAdv/Source/TP_VehicleAdv/Variant_TimeTrial/TimeTrialPlayerController.h
@@ -30,12 +30,8 @@ protected:
 	UPROPERTY(EditAnywhere, Category="Input|Input Mappings")
 	TArray<UInputMappingContext*> MobileExcludedMappingContexts;
 
-	/** If true, a mobile controls widget will be used for touch input */
-	UPROPERTY(EditAnywhere, Category="Input|Touch Controls")
-	bool bUseMobileControlsWidget = false;
-
 	/** Mobile controls widget to spawn */
-	UPROPERTY(EditAnywhere, Category="Input|Touch Controls", meta = (EditCondition = "bUseMobileControlsWidget"))
+	UPROPERTY(EditAnywhere, Category="Input|Touch Controls")
 	TSubclassOf<UUserWidget> MobileControlsWidgetClass;
 
 	/** Pointer to the mobile controls widget */
