Templates
- TP_FirstPerson: Fixed a crash due to custom EnvQuery not being able to generate samples when the player pawn dies.
- TP_FirstPerson: Fixed an issue with AnimBPs not working correctly after recent animation asset changes.
- TP_FirstPersonBP: Added extra validity railguards to Shooter NPC StateTree

#rnx

diff --git a/Templates/TP_FirstPerson/Source/TP_FirstPerson/Variant_Shooter/AI/EnvQueryContext_Target.cpp b/Templates/TP_FirstPerson/Source/TP_FirstPerson/Variant_Shooter/AI/EnvQueryContext_Target.cpp
--- a/Templates/TP_FirstPerson/Source/TP_FirstPerson/Variant_Shooter/AI/EnvQueryContext_Target.cpp
+++ b/Templates/TP_FirstPerson/Source/TP_FirstPerson/Variant_Shooter/AI/EnvQueryContext_Target.cpp
@@ -16,6 +16,11 @@ void UEnvQueryContext_Target::ProvideContext(FEnvQueryInstance& QueryInstance, F
 		{
 			// add the controller's target actor to the context
 			UEnvQueryItemType_Actor::SetContextHelper(ContextData, Controller->GetCurrentTarget());
+
+		} else {
+
+			// if for any reason there's no target, default to the controller
+			UEnvQueryItemType_Actor::SetContextHelper(ContextData, Controller);
 		}
 	}
 
