TP_FirstPerson, TP_FirstPersonBP
- Jump pads now ignore dead characters and NPCs
- Characters can no longer be controlled through mobile touch controls when dead
- First person camera mesh is no longer visible while death camera is active
- Cleared a stale reference to the old touch interface controls asset.
- Implemented force touch controls flag for non-mobile platforms.
#rnx

#virtualized

diff --git a/Templates/TP_FirstPerson/Source/TP_FirstPerson/TP_FirstPersonPlayerController.cpp b/Templates/TP_FirstPerson/Source/TP_FirstPerson/TP_FirstPersonPlayerController.cpp
--- a/Templates/TP_FirstPerson/Source/TP_FirstPerson/TP_FirstPersonPlayerController.cpp
+++ b/Templates/TP_FirstPerson/Source/TP_FirstPerson/TP_FirstPersonPlayerController.cpp
@@ -22,7 +22,7 @@ void ATP_FirstPersonPlayerController::BeginPlay()
 
 	
 	// only spawn touch controls on local player controllers
-	if (SVirtualJoystick::ShouldDisplayTouchInterface() && IsLocalPlayerController())
+	if (ShouldUseTouchControls() && IsLocalPlayerController())
 	{
 		// spawn the mobile controls widget
 		MobileControlsWidget = CreateWidget<UUserWidget>(this, MobileControlsWidgetClass);
@@ -57,7 +57,7 @@ void ATP_FirstPersonPlayerController::SetupInputComponent()
 			}
 
 			// only add these IMCs if we're not using mobile touch input
-			if (!SVirtualJoystick::ShouldDisplayTouchInterface())
+			if (!ShouldUseTouchControls())
 			{
 				for (UInputMappingContext* CurrentContext : MobileExcludedMappingContexts)
 				{
@@ -68,3 +68,9 @@ void ATP_FirstPersonPlayerController::SetupInputComponent()
 	}
 	
 }
+
+bool ATP_FirstPersonPlayerController::ShouldUseTouchControls() const
+{
+	// are we on a mobile platform? Should we force touch?
+	return SVirtualJoystick::ShouldDisplayTouchInterface() || bForceTouchControls;
+}
diff --git a/Templates/TP_FirstPerson/Source/TP_FirstPerson/TP_FirstPersonPlayerController.h b/Templates/TP_FirstPerson/Source/TP_FirstPerson/TP_FirstPersonPlayerController.h
--- a/Templates/TP_FirstPerson/Source/TP_FirstPerson/TP_FirstPersonPlayerController.h
+++ b/Templates/TP_FirstPerson/Source/TP_FirstPerson/TP_FirstPersonPlayerController.h
@@ -14,7 +14,7 @@ class UUserWidget;
  *  Manages the input mapping context.
  *  Overrides the Player Camera Manager class.
  */
-UCLASS(abstract)
+UCLASS(abstract, config="Game")
 class TP_FIRSTPERSON_API ATP_FirstPersonPlayerController : public APlayerController
 {
 	GENERATED_BODY()
@@ -41,10 +41,16 @@ protected:
 	/** Pointer to the mobile controls widget */
 	TObjectPtr<UUserWidget> MobileControlsWidget;
 
+	/** If true, the player will use UMG touch controls even if not playing on mobile platforms */
+	UPROPERTY(EditAnywhere, Config, Category = "Input|Touch Controls")
+	bool bForceTouchControls = false;
+
 	/** Gameplay initialization */
 	virtual void BeginPlay() override;
 
 	/** Input mapping context setup */
 	virtual void SetupInputComponent() override;
 
+	/** Returns true if the player should use UMG touch controls */
+	bool ShouldUseTouchControls() const;
 };
diff --git a/Templates/TP_FirstPerson/Source/TP_FirstPerson/Variant_Horror/HorrorPlayerController.cpp b/Templates/TP_FirstPerson/Source/TP_FirstPerson/Variant_Horror/HorrorPlayerController.cpp
--- a/Templates/TP_FirstPerson/Source/TP_FirstPerson/Variant_Horror/HorrorPlayerController.cpp
+++ b/Templates/TP_FirstPerson/Source/TP_FirstPerson/Variant_Horror/HorrorPlayerController.cpp
@@ -22,7 +22,7 @@ void AHorrorPlayerController::BeginPlay()
 	Super::BeginPlay();
 
 	// only spawn touch controls on local player controllers
-	if (SVirtualJoystick::ShouldDisplayTouchInterface() && IsLocalPlayerController())
+	if (ShouldUseTouchControls() && IsLocalPlayerController())
 	{
 		// spawn the mobile controls widget
 		MobileControlsWidget = CreateWidget<UUserWidget>(this, MobileControlsWidgetClass);
@@ -80,7 +80,7 @@ void AHorrorPlayerController::SetupInputComponent()
 			}
 
 			// only add these IMCs if we're not using mobile touch input
-			if (!SVirtualJoystick::ShouldDisplayTouchInterface())
+			if (!ShouldUseTouchControls())
 			{
 				for (UInputMappingContext* CurrentContext : MobileExcludedMappingContexts)
 				{
@@ -90,3 +90,9 @@ void AHorrorPlayerController::SetupInputComponent()
 		}
 	}	
 }
+
+bool AHorrorPlayerController::ShouldUseTouchControls() const
+{
+	// are we on a mobile platform? Should we force touch?
+	return SVirtualJoystick::ShouldDisplayTouchInterface() || bForceTouchControls;
+}
diff --git a/Templates/TP_FirstPerson/Source/TP_FirstPerson/Variant_Horror/HorrorPlayerController.h b/Templates/TP_FirstPerson/Source/TP_FirstPerson/Variant_Horror/HorrorPlayerController.h
--- a/Templates/TP_FirstPerson/Source/TP_FirstPerson/Variant_Horror/HorrorPlayerController.h
+++ b/Templates/TP_FirstPerson/Source/TP_FirstPerson/Variant_Horror/HorrorPlayerController.h
@@ -14,7 +14,7 @@ class UHorrorUI;
  *  Manages input mappings
  *  Manages UI
  */
-UCLASS(abstract)
+UCLASS(abstract, config="Game")
 class TP_FIRSTPERSON_API AHorrorPlayerController : public APlayerController
 {
 	GENERATED_BODY()
@@ -50,6 +50,10 @@ protected:
 	/** Pointer to the mobile controls widget */
 	TObjectPtr<UUserWidget> MobileControlsWidget;
 
+	/** If true, the player will use UMG touch controls even if not playing on mobile platforms */
+	UPROPERTY(EditAnywhere, Config, Category = "Input|Touch Controls")
+	bool bForceTouchControls = false;
+
 	/** Gameplay Initialization */
 	virtual void BeginPlay() override;
 
@@ -59,4 +63,6 @@ protected:
 	/** Input mapping context setup */
 	virtual void SetupInputComponent() override;
 
+	/** Returns true if the player should use UMG touch controls */
+	bool ShouldUseTouchControls() const;
 };
diff --git a/Templates/TP_FirstPerson/Source/TP_FirstPerson/Variant_Shooter/AI/ShooterNPC.cpp b/Templates/TP_FirstPerson/Source/TP_FirstPerson/Variant_Shooter/AI/ShooterNPC.cpp
--- a/Templates/TP_FirstPerson/Source/TP_FirstPerson/Variant_Shooter/AI/ShooterNPC.cpp
+++ b/Templates/TP_FirstPerson/Source/TP_FirstPerson/Variant_Shooter/AI/ShooterNPC.cpp
@@ -159,6 +159,9 @@ void AShooterNPC::Die()
 	// raise the dead flag
 	bIsDead = true;
 
+	// grant the death tag to the character
+	Tags.Add(DeathTag);
+
 	// increment the team score
 	if (AShooterGameMode* GM = Cast<AShooterGameMode>(GetWorld()->GetAuthGameMode()))
 	{
diff --git a/Templates/TP_FirstPerson/Source/TP_FirstPerson/Variant_Shooter/AI/ShooterNPC.h b/Templates/TP_FirstPerson/Source/TP_FirstPerson/Variant_Shooter/AI/ShooterNPC.h
--- a/Templates/TP_FirstPerson/Source/TP_FirstPerson/Variant_Shooter/AI/ShooterNPC.h
+++ b/Templates/TP_FirstPerson/Source/TP_FirstPerson/Variant_Shooter/AI/ShooterNPC.h
@@ -41,6 +41,10 @@ protected:
 	UPROPERTY(EditAnywhere, Category="Team")
 	uint8 TeamByte = 1;
 
+	/** Actor tag to grant this character when it dies */
+	UPROPERTY(EditAnywhere, Category="Team")
+	FName DeathTag = FName("Dead");
+
 	/** Pointer to the equipped weapon */
 	TObjectPtr<AShooterWeapon> Weapon;
 
diff --git a/Templates/TP_FirstPerson/Source/TP_FirstPerson/Variant_Shooter/ShooterCharacter.cpp b/Templates/TP_FirstPerson/Source/TP_FirstPerson/Variant_Shooter/ShooterCharacter.cpp
--- a/Templates/TP_FirstPerson/Source/TP_FirstPerson/Variant_Shooter/ShooterCharacter.cpp
+++ b/Templates/TP_FirstPerson/Source/TP_FirstPerson/Variant_Shooter/ShooterCharacter.cpp
@@ -82,10 +82,46 @@ float AShooterCharacter::TakeDamage(float Damage, struct FDamageEvent const& Dam
 	return Damage;
 }
 
+void AShooterCharacter::DoAim(float Yaw, float Pitch)
+{
+	// only route inputs if the character is not dead
+	if (!IsDead())
+	{
+		Super::DoAim(Yaw, Pitch);
+	}
+}
+
+void AShooterCharacter::DoMove(float Right, float Forward)
+{
+	// only route inputs if the character is not dead
+	if (!IsDead())
+	{
+		Super::DoMove(Right, Forward);
+	}
+}
+
+void AShooterCharacter::DoJumpStart()
+{
+	// only route inputs if the character is not dead
+	if (!IsDead())
+	{
+		Super::DoJumpStart();
+	}
+}
+
+void AShooterCharacter::DoJumpEnd()
+{
+	// only route inputs if the character is not dead
+	if (!IsDead())
+	{
+		Super::DoJumpEnd();
+	}
+}
+
 void AShooterCharacter::DoStartFiring()
 {
 	// fire the current weapon
-	if (CurrentWeapon)
+	if (CurrentWeapon && !IsDead())
 	{
 		CurrentWeapon->StartFiring();
 	}
@@ -94,7 +130,7 @@ void AShooterCharacter::DoStartFiring()
 void AShooterCharacter::DoStopFiring()
 {
 	// stop firing the current weapon
-	if (CurrentWeapon)
+	if (CurrentWeapon && !IsDead())
 	{
 		CurrentWeapon->StopFiring();
 	}
@@ -103,7 +139,7 @@ void AShooterCharacter::DoStopFiring()
 void AShooterCharacter::DoSwitchWeapon()
 {
 	// ensure we have at least two weapons two switch between
-	if (OwnedWeapons.Num() > 1)
+	if (OwnedWeapons.Num() > 1 && !IsDead())
 	{
 		// deactivate the old weapon
 		CurrentWeapon->DeactivateWeapon();
@@ -145,7 +181,7 @@ void AShooterCharacter::AttachWeaponMeshes(AShooterWeapon* Weapon)
 
 void AShooterCharacter::PlayFiringMontage(UAnimMontage* Montage)
 {
-	
+	// stub
 }
 
 void AShooterCharacter::AddWeaponRecoil(float Recoil)
@@ -259,6 +295,9 @@ void AShooterCharacter::Die()
 	{
 		GM->IncrementTeamScore(TeamByte);
 	}
+
+	// grant the death tag to the character
+	Tags.Add(DeathTag);
 		
 	// stop character movement
 	GetCharacterMovement()->StopMovementImmediately();
@@ -281,3 +320,9 @@ void AShooterCharacter::OnRespawn()
 	// destroy the character to force the PC to respawn
 	Destroy();
 }
+
+bool AShooterCharacter::IsDead() const
+{
+	// the character is dead if their current HP drops to zero
+	return CurrentHP <= 0.0f;
+}
diff --git a/Templates/TP_FirstPerson/Source/TP_FirstPerson/Variant_Shooter/ShooterCharacter.h b/Templates/TP_FirstPerson/Source/TP_FirstPerson/Variant_Shooter/ShooterCharacter.h
--- a/Templates/TP_FirstPerson/Source/TP_FirstPerson/Variant_Shooter/ShooterCharacter.h
+++ b/Templates/TP_FirstPerson/Source/TP_FirstPerson/Variant_Shooter/ShooterCharacter.h
@@ -62,6 +62,10 @@ protected:
 	UPROPERTY(EditAnywhere, Category="Team")
 	uint8 TeamByte = 0;
 
+	/** Actor tag to grant this character when it dies */
+	UPROPERTY(EditAnywhere, Category="Team")
+	FName DeathTag = FName("Dead");
+
 	/** List of weapons picked up by the character */
 	TArray<AShooterWeapon*> OwnedWeapons;
 
@@ -104,6 +108,18 @@ public:
 
 public:
 
+	/** Handles aim inputs from either controls or UI interfaces */
+	virtual void DoAim(float Yaw, float Pitch) override;
+
+	/** Handles move inputs from either controls or UI interfaces */
+	virtual void DoMove(float Right, float Forward)  override;
+
+	/** Handles jump start inputs from either controls or UI interfaces */
+	virtual void DoJumpStart()  override;
+
+	/** Handles jump end inputs from either controls or UI interfaces */
+	virtual void DoJumpEnd()  override;
+
 	/** Handles start firing input */
 	UFUNCTION(BlueprintCallable, Category="Input")
 	void DoStartFiring();
@@ -163,4 +179,9 @@ protected:
 
 	/** Called from the respawn timer to destroy this character and force the PC to respawn */
 	void OnRespawn();
+
+public:
+
+	/** Returns true if the character is dead */
+	bool IsDead() const;
 };
diff --git a/Templates/TP_FirstPerson/Source/TP_FirstPerson/Variant_Shooter/ShooterPlayerController.cpp b/Templates/TP_FirstPerson/Source/TP_FirstPerson/Variant_Shooter/ShooterPlayerController.cpp
--- a/Templates/TP_FirstPerson/Source/TP_FirstPerson/Variant_Shooter/ShooterPlayerController.cpp
+++ b/Templates/TP_FirstPerson/Source/TP_FirstPerson/Variant_Shooter/ShooterPlayerController.cpp
@@ -19,7 +19,7 @@ void AShooterPlayerController::BeginPlay()
 	// only spawn touch controls on local player controllers
 	if (IsLocalPlayerController())
 	{
-		if (SVirtualJoystick::ShouldDisplayTouchInterface())
+		if (ShouldUseTouchControls())
 		{
 			// spawn the mobile controls widget
 			MobileControlsWidget = CreateWidget<UUserWidget>(this, MobileControlsWidgetClass);
@@ -66,7 +66,7 @@ void AShooterPlayerController::SetupInputComponent()
 			}
 
 			// only add these IMCs if we're not using mobile touch input
-			if (!SVirtualJoystick::ShouldDisplayTouchInterface())
+			if (!ShouldUseTouchControls())
 			{
 				for (UInputMappingContext* CurrentContext : MobileExcludedMappingContexts)
 				{
@@ -140,3 +140,9 @@ void AShooterPlayerController::OnPawnDamaged(float LifePercent)
 		BulletCounterUI->BP_Damaged(LifePercent);
 	}
 }
+
+bool AShooterPlayerController::ShouldUseTouchControls() const
+{
+	// are we on a mobile platform? Should we force touch?
+	return SVirtualJoystick::ShouldDisplayTouchInterface() || bForceTouchControls;
+}
diff --git a/Templates/TP_FirstPerson/Source/TP_FirstPerson/Variant_Shooter/ShooterPlayerController.h b/Templates/TP_FirstPerson/Source/TP_FirstPerson/Variant_Shooter/ShooterPlayerController.h
--- a/Templates/TP_FirstPerson/Source/TP_FirstPerson/Variant_Shooter/ShooterPlayerController.h
+++ b/Templates/TP_FirstPerson/Source/TP_FirstPerson/Variant_Shooter/ShooterPlayerController.h
@@ -15,7 +15,7 @@ class UShooterBulletCounterUI;
  *  Manages input mappings
  *  Respawns the player pawn when it's destroyed
  */
-UCLASS(abstract)
+UCLASS(abstract, config="Game")
 class TP_FIRSTPERSON_API AShooterPlayerController : public APlayerController
 {
 	GENERATED_BODY()
@@ -37,6 +37,10 @@ protected:
 	/** Pointer to the mobile controls widget */
 	TObjectPtr<UUserWidget> MobileControlsWidget;
 
+	/** If true, the player will use UMG touch controls even if not playing on mobile platforms */
+	UPROPERTY(EditAnywhere, Config, Category = "Input|Touch Controls")
+	bool bForceTouchControls = false;
+
 	/** Character class to respawn when the possessed pawn is destroyed */
 	UPROPERTY(EditAnywhere, Category="Shooter|Respawn")
 	TSubclassOf<AShooterCharacter> CharacterClass;
@@ -74,4 +78,7 @@ protected:
 	/** Called when the possessed pawn is damaged */
 	UFUNCTION()
 	void OnPawnDamaged(float LifePercent);
+
+	/** Returns true if the player should use UMG touch controls */
+	bool ShouldUseTouchControls() const;
 };
