TP_FirstPerson, TP_FirstPersonBP
- Implemented NPC Spawners with editable respawn count.
- Replaced NPCs on Shooter variant map for spawners. Respawn count is set to 1 to keep parity with previous versions.
- Shooter variant map on TP_FirstPersonBP still need to be replaced.
#rnx

#virtualized

diff --git a/Templates/TP_FirstPerson/Source/TP_FirstPerson/Variant_Shooter/AI/ShooterNPCSpawner.cpp b/Templates/TP_FirstPerson/Source/TP_FirstPerson/Variant_Shooter/AI/ShooterNPCSpawner.cpp
new file mode 100644
--- /dev/null
+++ b/Templates/TP_FirstPerson/Source/TP_FirstPerson/Variant_Shooter/AI/ShooterNPCSpawner.cpp
@@ -0,0 +1,85 @@
+// Copyright Epic Games, Inc. All Rights Reserved.
+
+
+#include "Variant_Shooter/AI/ShooterNPCSpawner.h"
+#include "Engine/World.h"
+#include "Components/SceneComponent.h"
+#include "Components/CapsuleComponent.h"
+#include "Components/ArrowComponent.h"
+#include "TimerManager.h"
+#include "ShooterNPC.h"
+
+// Sets default values
+AShooterNPCSpawner::AShooterNPCSpawner()
+{
+	PrimaryActorTick.bCanEverTick = false;
+
+	// create the root
+	RootComponent = CreateDefaultSubobject<USceneComponent>(TEXT("Root"));
+
+	// create the reference spawn capsule
+	SpawnCapsule = CreateDefaultSubobject<UCapsuleComponent>(TEXT("Spawn Capsule"));
+	SpawnCapsule->SetupAttachment(RootComponent);
+
+	SpawnCapsule->SetRelativeLocation(FVector(0.0f, 0.0f, 90.0f));
+	SpawnCapsule->SetCapsuleSize(35.0f, 90.0f);
+	SpawnCapsule->SetCollisionProfileName(FName("NoCollision"));
+
+	SpawnDirection = CreateDefaultSubobject<UArrowComponent>(TEXT("Spawn Direction"));
+	SpawnDirection->SetupAttachment(RootComponent);
+}
+
+void AShooterNPCSpawner::BeginPlay()
+{
+	Super::BeginPlay();
+	
+	// should we spawn an NPC right away?
+	if (bShouldSpawnNPCsImmediately)
+	{
+		// schedule the first NPC spawn
+		GetWorld()->GetTimerManager().SetTimer(SpawnTimer, this, &AShooterNPCSpawner::SpawnNPC, InitialSpawnDelay);
+	}
+}
+
+void AShooterNPCSpawner::EndPlay(EEndPlayReason::Type EndPlayReason)
+{
+	Super::EndPlay(EndPlayReason);
+
+	// clear the spawn timer
+	GetWorld()->GetTimerManager().ClearTimer(SpawnTimer);
+}
+
+void AShooterNPCSpawner::SpawnNPC()
+{
+	// ensure the NPC class is valid
+	if (IsValid(NPCClass))
+	{
+		// spawn the NPC at the reference capsule's transform
+		FActorSpawnParameters SpawnParams;
+		SpawnParams.SpawnCollisionHandlingOverride = ESpawnActorCollisionHandlingMethod::AdjustIfPossibleButAlwaysSpawn;
+
+		AShooterNPC* SpawnedNPC = GetWorld()->SpawnActor<AShooterNPC>(NPCClass, SpawnCapsule->GetComponentTransform(), SpawnParams);
+
+		// was the NPC successfully created?
+		if (SpawnedNPC)
+		{
+			// subscribe to the death delegate
+			SpawnedNPC->OnPawnDeath.AddDynamic(this, &AShooterNPCSpawner::OnNPCDied);
+		}
+	}
+}
+
+void AShooterNPCSpawner::OnNPCDied()
+{
+	// decrease the spawn counter
+	--SpawnCount;
+
+	// is this the last NPC we should spawn?
+	if (SpawnCount <= 0)
+	{
+		return;
+	}
+
+	// schedule the next NPC spawn
+	GetWorld()->GetTimerManager().SetTimer(SpawnTimer, this, &AShooterNPCSpawner::SpawnNPC, RespawnDelay);
+}
diff --git a/Templates/TP_FirstPerson/Source/TP_FirstPerson/Variant_Shooter/AI/ShooterNPCSpawner.h b/Templates/TP_FirstPerson/Source/TP_FirstPerson/Variant_Shooter/AI/ShooterNPCSpawner.h
new file mode 100644
--- /dev/null
+++ b/Templates/TP_FirstPerson/Source/TP_FirstPerson/Variant_Shooter/AI/ShooterNPCSpawner.h
@@ -0,0 +1,75 @@
+// Copyright Epic Games, Inc. All Rights Reserved.
+
+#pragma once
+
+#include "CoreMinimal.h"
+#include "GameFramework/Actor.h"
+#include "ShooterNPCSpawner.generated.h"
+
+class UCapsuleComponent;
+class UArrowComponent;
+class AShooterNPC;
+
+/**
+ *  A basic Actor in charge of spawning Shooter NPCs and monitoring their deaths.
+ *  NPCs will be spawned one by one, and the spawner will wait until it dies before spawning a new one.
+ */
+UCLASS()
+class TP_FIRSTPERSON_API AShooterNPCSpawner : public AActor
+{
+	GENERATED_BODY()
+	
+	UPROPERTY(VisibleAnywhere, BlueprintReadOnly, Category="Components", meta = (AllowPrivateAccess = "true"))
+	UCapsuleComponent* SpawnCapsule;
+
+	UPROPERTY(VisibleAnywhere, BlueprintReadOnly, Category = "Components", meta = (AllowPrivateAccess = "true"))
+	UArrowComponent* SpawnDirection;
+
+protected:
+
+	/** Type of NPC to spawn */
+	UPROPERTY(EditAnywhere, BlueprintReadOnly, Category="NPC Spawner")
+	TSubclassOf<AShooterNPC> NPCClass;
+
+	/** If true, the first NPC will be spawned as soon as the game starts */
+	UPROPERTY(EditAnywhere, BlueprintReadOnly, Category="NPC Spawner")
+	bool bShouldSpawnNPCsImmediately = true;
+
+	/** Time to wait before spawning the first NPC on game start */
+	UPROPERTY(EditAnywhere, BlueprintReadOnly, Category="NPC Spawner", meta = (ClampMin = 0, ClampMax = 10))
+	float InitialSpawnDelay = 5.0f;
+
+	/** Number of NPCs to spawn */
+	UPROPERTY(EditAnywhere, BlueprintReadOnly, Category="NPC Spawner", meta = (ClampMin = 0, ClampMax = 100))
+	int32 SpawnCount = 1;
+
+	/** Time to wait before spawning the next NPC after the current one dies */
+	UPROPERTY(EditAnywhere, BlueprintReadOnly, Category="NPC Spawner", meta = (ClampMin = 0, ClampMax = 10))
+	float RespawnDelay = 5.0f;
+
+	/** Timer to spawn NPCs after a delay */
+	FTimerHandle SpawnTimer;
+
+public:	
+	
+	/** Constructor */
+	AShooterNPCSpawner();
+
+public:
+
+	/** Initialization */
+	virtual void BeginPlay() override;
+
+	/** Cleanup */
+	virtual void EndPlay(EEndPlayReason::Type EndPlayReason) override;
+
+protected:
+
+	/** Spawn an NPC and subscribe to its death event */
+	void SpawnNPC();
+
+	/** Called when the spawned NPC has died */
+	UFUNCTION()
+	void OnNPCDied();
+
+};
