TP_FirstPerson
- Fixed an issue with spawned NPCs not automatically starting AI logic.

#rnx

diff --git a/Templates/TP_FirstPerson/Source/TP_FirstPerson/Variant_Shooter/AI/ShooterAIController.cpp b/Templates/TP_FirstPerson/Source/TP_FirstPerson/Variant_Shooter/AI/ShooterAIController.cpp
--- a/Templates/TP_FirstPerson/Source/TP_FirstPerson/Variant_Shooter/AI/ShooterAIController.cpp
+++ b/Templates/TP_FirstPerson/Source/TP_FirstPerson/Variant_Shooter/AI/ShooterAIController.cpp
@@ -12,6 +12,7 @@ AShooterAIController::AShooterAIController()
 {
 	// create the StateTree component
 	StateTreeAI = CreateDefaultSubobject<UStateTreeAIComponent>(TEXT("StateTreeAI"));
+	StateTreeAI->SetStartLogicAutomatically(false);
 
 	// create the AI perception component. It will be configured in BP
 	AIPerception = CreateDefaultSubobject<UAIPerceptionComponent>(TEXT("AIPerception"));
@@ -33,6 +34,9 @@ void AShooterAIController::OnPossess(APawn* InPawn)
 
 		// subscribe to the pawn's OnDeath delegate
 		NPC->OnPawnDeath.AddDynamic(this, &AShooterAIController::OnPawnDeath);
+
+		// start AI logic
+		StateTreeAI->StartLogic();
 	}
 }
 
