Templates
- Fixed SetupInputComponent crash on some template variants.
- Fixed a typo on the description of the Intro to UE template.
#rnx

diff --git a/Templates/TP_FirstPerson/Source/TP_FirstPerson/Variant_Shooter/ShooterPlayerController.cpp b/Templates/TP_FirstPerson/Source/TP_FirstPerson/Variant_Shooter/ShooterPlayerController.cpp
--- a/Templates/TP_FirstPerson/Source/TP_FirstPerson/Variant_Shooter/ShooterPlayerController.cpp
+++ b/Templates/TP_FirstPerson/Source/TP_FirstPerson/Variant_Shooter/ShooterPlayerController.cpp
@@ -54,6 +54,8 @@ void AShooterPlayerController::BeginPlay()
 
 void AShooterPlayerController::SetupInputComponent()
 {
+	Super::SetupInputComponent();
+
 	// only add IMCs for local player controllers
 	if (IsLocalPlayerController())
 	{
@@ -102,7 +104,10 @@ void AShooterPlayerController::OnPossess(APawn* InPawn)
 void AShooterPlayerController::OnPawnDestroyed(AActor* DestroyedActor)
 {
 	// reset the bullet counter HUD
-	BulletCounterUI->BP_UpdateBulletCounter(0, 0);
+	if (IsValid(BulletCounterUI))
+	{
+		BulletCounterUI->BP_UpdateBulletCounter(0, 0);
+	}
 
 	// find the player start
 	TArray<AActor*> ActorList;
diff --git a/Templates/TP_FirstPerson/Source/TP_FirstPerson/Variant_Shooter/ShooterPlayerController.h b/Templates/TP_FirstPerson/Source/TP_FirstPerson/Variant_Shooter/ShooterPlayerController.h
--- a/Templates/TP_FirstPerson/Source/TP_FirstPerson/Variant_Shooter/ShooterPlayerController.h
+++ b/Templates/TP_FirstPerson/Source/TP_FirstPerson/Variant_Shooter/ShooterPlayerController.h
@@ -55,6 +55,7 @@ protected:
 	FName PlayerPawnTag = FName("Player");
 
 	/** Pointer to the bullet counter UI widget */
+	UPROPERTY()
 	TObjectPtr<UShooterBulletCounterUI> BulletCounterUI;
 
 protected:
diff --git a/Templates/TP_ThirdPerson/Source/TP_ThirdPerson/Variant_Combat/CombatPlayerController.cpp b/Templates/TP_ThirdPerson/Source/TP_ThirdPerson/Variant_Combat/CombatPlayerController.cpp
--- a/Templates/TP_ThirdPerson/Source/TP_ThirdPerson/Variant_Combat/CombatPlayerController.cpp
+++ b/Templates/TP_ThirdPerson/Source/TP_ThirdPerson/Variant_Combat/CombatPlayerController.cpp
@@ -39,6 +39,8 @@ void ACombatPlayerController::BeginPlay()
 
 void ACombatPlayerController::SetupInputComponent()
 {
+	Super::SetupInputComponent();
+
 	// only add IMCs for local player controllers
 	if (IsLocalPlayerController())
 	{
diff --git a/Templates/TP_ThirdPerson/Source/TP_ThirdPerson/Variant_Platforming/PlatformingPlayerController.cpp b/Templates/TP_ThirdPerson/Source/TP_ThirdPerson/Variant_Platforming/PlatformingPlayerController.cpp
--- a/Templates/TP_ThirdPerson/Source/TP_ThirdPerson/Variant_Platforming/PlatformingPlayerController.cpp
+++ b/Templates/TP_ThirdPerson/Source/TP_ThirdPerson/Variant_Platforming/PlatformingPlayerController.cpp
@@ -39,6 +39,8 @@ void APlatformingPlayerController::BeginPlay()
 
 void APlatformingPlayerController::SetupInputComponent()
 {
+	Super::SetupInputComponent();
+
 	// only add IMCs for local player controllers
 	if (IsLocalPlayerController())
 	{
diff --git a/Templates/TP_UEIntro_BP/Config/TemplateDefs.ini b/Templates/TP_UEIntro_BP/Config/TemplateDefs.ini
--- a/Templates/TP_UEIntro_BP/Config/TemplateDefs.ini
+++ b/Templates/TP_UEIntro_BP/Config/TemplateDefs.ini
@@ -9,7 +9,7 @@ Categories=SIM
 SortKey="0"
 
 LocalizedDisplayNames=(Language="en",Text="Intro To Unreal")
-LocalizedDescriptions=(Language="en",Text="This template takes you through the basics of the engine, including movement, adding objects to a level and the various types of asset you can create.")
+LocalizedDescriptions=(Language="en",Text="This template takes you through the basics of the engine, including movement, adding objects to a level and the various types of assets you can create.")
 
 ; This project uses a thumbnail as the project icon
 bThumbnailAsIcon=true
