Fix Deprecation warning.

#rnx
#FailedValidations Preflight

diff --git a/Templates/TP_FirstPerson/Source/TP_FirstPerson/Variant_Shooter/AI/ShooterStateTreeUtility.h b/Templates/TP_FirstPerson/Source/TP_FirstPerson/Variant_Shooter/AI/ShooterStateTreeUtility.h
--- a/Templates/TP_FirstPerson/Source/TP_FirstPerson/Variant_Shooter/AI/ShooterStateTreeUtility.h
+++ b/Templates/TP_FirstPerson/Source/TP_FirstPerson/Variant_Shooter/AI/ShooterStateTreeUtility.h
@@ -22,11 +22,11 @@ struct FStateTreeLineOfSightToTargetConditionInstanceData
 	
 	/** Targeting character */
 	UPROPERTY(EditAnywhere, Category = "Context")
-	AShooterNPC* Character;
+	TObjectPtr<AShooterNPC> Character;
 
 	/** Target to check line of sight for */
 	UPROPERTY(EditAnywhere, Category = "Condition")
-	AActor* Target;
+	TObjectPtr<AActor> Target;
 
 	/** Max allowed line of sight cone angle, in degrees */
 	UPROPERTY(EditAnywhere, Category = "Condition")
@@ -40,7 +40,6 @@ struct FStateTreeLineOfSightToTargetConditionInstanceData
 	UPROPERTY(EditAnywhere, Category = "Condition")
 	bool bMustHaveLineOfSight = true;
 };
-STATETREE_POD_INSTANCEDATA(FStateTreeLineOfSightToTargetConditionInstanceData);
 
 /**
  *  StateTree condition to check if the character is grounded
diff --git a/Templates/TP_ThirdPerson/Source/TP_ThirdPerson/Variant_Combat/AI/CombatStateTreeUtility.h b/Templates/TP_ThirdPerson/Source/TP_ThirdPerson/Variant_Combat/AI/CombatStateTreeUtility.h
--- a/Templates/TP_ThirdPerson/Source/TP_ThirdPerson/Variant_Combat/AI/CombatStateTreeUtility.h
+++ b/Templates/TP_ThirdPerson/Source/TP_ThirdPerson/Variant_Combat/AI/CombatStateTreeUtility.h
@@ -22,13 +22,12 @@ struct FStateTreeCharacterGroundedConditionInstanceData
 	
 	/** Character to check grounded status on */
 	UPROPERTY(EditAnywhere, Category = "Context")
-	ACharacter* Character;
+	TObjectPtr<ACharacter> Character;
 
 	/** If true, the condition passes if the character is not grounded instead */
 	UPROPERTY(EditAnywhere, Category = "Condition")
 	bool bMustBeOnAir = false;
 };
-STATETREE_POD_INSTANCEDATA(FStateTreeCharacterGroundedConditionInstanceData);
 
 /**
  *  StateTree condition to check if the character is grounded
