Templates - TP_FirstPerson, TP_FirstPersonBP
- Arena Shooter: Fixed an issue that caused dead player characters to be moved by other actors.
- Arena Shooter: Fixed multiple AI issues that caused enemy NPCs to become unresponsive or fail to detect the player.

#rnx

#virtualized
#Synced-CL 52765683

diff --git a/Templates/TP_FirstPerson/Source/TP_FirstPerson/Variant_Shooter/AI/ShooterStateTreeUtility.cpp b/Templates/TP_FirstPerson/Source/TP_FirstPerson/Variant_Shooter/AI/ShooterStateTreeUtility.cpp
--- a/Templates/TP_FirstPerson/Source/TP_FirstPerson/Variant_Shooter/AI/ShooterStateTreeUtility.cpp
+++ b/Templates/TP_FirstPerson/Source/TP_FirstPerson/Variant_Shooter/AI/ShooterStateTreeUtility.cpp
@@ -80,30 +80,22 @@ FText FStateTreeLineOfSightToTargetCondition::GetDescription(const FGuid& ID, FS
 
 EStateTreeRunStatus FStateTreeFaceActorTask::EnterState(FStateTreeExecutionContext& Context, const FStateTreeTransitionResult& Transition) const
 {
-	// have we transitioned from another state?
-	if (Transition.ChangeType == EStateTreeStateChangeType::Changed)
-	{
-		// get the instance data
-		FInstanceDataType& InstanceData = Context.GetInstanceData(*this);
+	// get the instance data
+	FInstanceDataType& InstanceData = Context.GetInstanceData(*this);
 
-		// set the AI Controller's focus
-		InstanceData.Controller->SetFocus(InstanceData.ActorToFaceTowards);
-	}
+	// set the AI Controller's focus
+	InstanceData.Controller->SetFocus(InstanceData.ActorToFaceTowards);
 
 	return EStateTreeRunStatus::Running;
 }
 
 void FStateTreeFaceActorTask::ExitState(FStateTreeExecutionContext& Context, const FStateTreeTransitionResult& Transition) const
 {
-	// have we transitioned to another state?
-	if (Transition.ChangeType == EStateTreeStateChangeType::Changed)
-	{
-		// get the instance data
-		FInstanceDataType& InstanceData = Context.GetInstanceData(*this);
+	// get the instance data
+	FInstanceDataType& InstanceData = Context.GetInstanceData(*this);
 
-		// clear the AI Controller's focus
-		InstanceData.Controller->ClearFocus(EAIFocusPriority::Gameplay);
-	}
+	// clear the AI Controller's focus
+	InstanceData.Controller->ClearFocus(EAIFocusPriority::Gameplay);
 }
 
 #if WITH_EDITOR
@@ -117,30 +109,22 @@ FText FStateTreeFaceActorTask::GetDescription(const FGuid& ID, FStateTreeDataVie
 
 EStateTreeRunStatus FStateTreeFaceLocationTask::EnterState(FStateTreeExecutionContext& Context, const FStateTreeTransitionResult& Transition) const
 {
-	// have we transitioned from another state?
-	if (Transition.ChangeType == EStateTreeStateChangeType::Changed)
-	{
-		// get the instance data
-		FInstanceDataType& InstanceData = Context.GetInstanceData(*this);
+	// get the instance data
+	FInstanceDataType& InstanceData = Context.GetInstanceData(*this);
 
-		// set the AI Controller's focus
-		InstanceData.Controller->SetFocalPoint(InstanceData.FaceLocation);
-	}
+	// set the AI Controller's focus
+	InstanceData.Controller->SetFocalPoint(InstanceData.FaceLocation);
 
 	return EStateTreeRunStatus::Running;
 }
 
 void FStateTreeFaceLocationTask::ExitState(FStateTreeExecutionContext& Context, const FStateTreeTransitionResult& Transition) const
 {
-	// have we transitioned to another state?
-	if (Transition.ChangeType == EStateTreeStateChangeType::Changed)
-	{
-		// get the instance data
-		FInstanceDataType& InstanceData = Context.GetInstanceData(*this);
+	// get the instance data
+	FInstanceDataType& InstanceData = Context.GetInstanceData(*this);
 
-		// clear the AI Controller's focus
-		InstanceData.Controller->ClearFocus(EAIFocusPriority::Gameplay);
-	}
+	// clear the AI Controller's focus
+	InstanceData.Controller->ClearFocus(EAIFocusPriority::Gameplay);
 }
 
 #if WITH_EDITOR
@@ -154,15 +138,11 @@ FText FStateTreeFaceLocationTask::GetDescription(const FGuid& ID, FStateTreeData
 
 EStateTreeRunStatus FStateTreeSetRandomFloatTask::EnterState(FStateTreeExecutionContext& Context, const FStateTreeTransitionResult& Transition) const
 {
-	// have we transitioned to another state?
-	if (Transition.ChangeType == EStateTreeStateChangeType::Changed)
-	{
-		// get the instance data
-		FInstanceDataType& InstanceData = Context.GetInstanceData(*this);
+	// get the instance data
+	FInstanceDataType& InstanceData = Context.GetInstanceData(*this);
 
-		// calculate the output value
-		InstanceData.OutValue = FMath::RandRange(InstanceData.MinValue, InstanceData.MaxValue);
-	}
+	// calculate the output value
+	InstanceData.OutValue = FMath::RandRange(InstanceData.MinValue, InstanceData.MaxValue);
 
 	return EStateTreeRunStatus::Running;
 }
@@ -178,30 +158,22 @@ FText FStateTreeSetRandomFloatTask::GetDescription(const FGuid& ID, FStateTreeDa
 
 EStateTreeRunStatus FStateTreeShootAtTargetTask::EnterState(FStateTreeExecutionContext& Context, const FStateTreeTransitionResult& Transition) const
 {
-	// have we transitioned from another state?
-	if (Transition.ChangeType == EStateTreeStateChangeType::Changed)
-	{
-		// get the instance data
-		FInstanceDataType& InstanceData = Context.GetInstanceData(*this);
+	// get the instance data
+	FInstanceDataType& InstanceData = Context.GetInstanceData(*this);
 
-		// tell the character to shoot the target
-		InstanceData.Character->StartShooting(InstanceData.Target);
-	}
+	// tell the character to shoot the target
+	InstanceData.Character->StartShooting(InstanceData.Target);
 
 	return EStateTreeRunStatus::Running;
 }
 
 void FStateTreeShootAtTargetTask::ExitState(FStateTreeExecutionContext& Context, const FStateTreeTransitionResult& Transition) const
 {
-	// have we transitioned to another state?
-	if (Transition.ChangeType == EStateTreeStateChangeType::Changed)
-	{
-		// get the instance data
-		FInstanceDataType& InstanceData = Context.GetInstanceData(*this);
+	// get the instance data
+	FInstanceDataType& InstanceData = Context.GetInstanceData(*this);
 
-		// tell the character to stop shooting
-		InstanceData.Character->StopShooting();
-	}
+	// tell the character to stop shooting
+	InstanceData.Character->StopShooting();
 }
 
 #if WITH_EDITOR
@@ -213,144 +185,141 @@ FText FStateTreeShootAtTargetTask::GetDescription(const FGuid& ID, FStateTreeDat
 
 EStateTreeRunStatus FStateTreeSenseEnemiesTask::EnterState(FStateTreeExecutionContext& Context, const FStateTreeTransitionResult& Transition) const
 {
-	// have we transitioned from another state?
-	if (Transition.ChangeType == EStateTreeStateChangeType::Changed)
-	{
-		// get the instance data
-		FInstanceDataType& InstanceData = Context.GetInstanceData(*this);
+	// get the instance data
+	FInstanceDataType& InstanceData = Context.GetInstanceData(*this);
 
-		// bind the perception updated delegate on the controller
-		InstanceData.Controller->OnShooterPerceptionUpdated.BindLambda(
-			[WeakContext = Context.MakeWeakExecutionContext()](AActor* SensedActor, const FAIStimulus& Stimulus)
-			{
-				// get the instance data inside the lambda
-				const FStateTreeStrongExecutionContext StrongContext = WeakContext.MakeStrongExecutionContext();
-				if (FInstanceDataType* LambdaInstanceData = StrongContext.GetInstanceDataPtr<FInstanceDataType>())
+	// bind the perception updated delegate on the controller
+	InstanceData.Controller->OnShooterPerceptionUpdated.BindLambda(
+		[WeakContext = Context.MakeWeakExecutionContext()](AActor* SensedActor, const FAIStimulus& Stimulus)
+		{
+			// get the instance data inside the lambda
+			const FStateTreeStrongExecutionContext StrongContext = WeakContext.MakeStrongExecutionContext();
+			if (FInstanceDataType* LambdaInstanceData = StrongContext.GetInstanceDataPtr<FInstanceDataType>())
+			{				
+				if (Stimulus.Tag == LambdaInstanceData->SenseTag || SensedActor->ActorHasTag(LambdaInstanceData->SenseTag))
 				{
-					if (SensedActor->ActorHasTag(LambdaInstanceData->SenseTag))
+					bool bDirectLOS = false;
+
+					// calculate the direction of the stimulus
+					const FVector StimulusDir = (Stimulus.StimulusLocation - LambdaInstanceData->Character->GetActorLocation()).GetSafeNormal();
+
+					// infer the angle from the dot product between the character facing and the stimulus direction
+					const float DirDot = FVector::DotProduct(StimulusDir, LambdaInstanceData->Character->GetActorForwardVector());
+					const float MaxDot = FMath::Cos(FMath::DegreesToRadians(LambdaInstanceData->DirectLineOfSightCone));
+
+					// is the direction within our perception cone?
+					if (DirDot >= MaxDot)
 					{
-						bool bDirectLOS = false;
+						// run a line trace between the character and the sensed actor
+						FCollisionQueryParams QueryParams;
+						QueryParams.AddIgnoredActor(LambdaInstanceData->Character);
+						QueryParams.AddIgnoredActor(SensedActor);
 
-						// calculate the direction of the stimulus
-						const FVector StimulusDir = (Stimulus.StimulusLocation - LambdaInstanceData->Character->GetActorLocation()).GetSafeNormal();
+						FHitResult OutHit;
 
-						// infer the angle from the dot product between the character facing and the stimulus direction
-						const float DirDot = FVector::DotProduct(StimulusDir, LambdaInstanceData->Character->GetActorForwardVector());
-						const float MaxDot = FMath::Cos(FMath::DegreesToRadians(LambdaInstanceData->DirectLineOfSightCone));
+						// we have direct line of sight if this trace is unobstructed
+						bDirectLOS = !LambdaInstanceData->Character->GetWorld()->LineTraceSingleByChannel(OutHit, LambdaInstanceData->Character->GetActorLocation(), SensedActor->GetActorLocation(), ECC_Visibility, QueryParams);
 
-						// is the direction within our perception cone?
-						if (DirDot >= MaxDot)
-						{
-							// run a line trace between the character and the sensed actor
-							FCollisionQueryParams QueryParams;
-							QueryParams.AddIgnoredActor(LambdaInstanceData->Character);
-							QueryParams.AddIgnoredActor(SensedActor);
+					}
 
-							FHitResult OutHit;
+					// check if we have a direct line of sight to the stimulus
+					if (bDirectLOS)
+					{
+						// set the controller's target
+						LambdaInstanceData->Controller->SetCurrentTarget(SensedActor);
 
-							// we have direct line of sight if this trace is unobstructed
-							bDirectLOS = !LambdaInstanceData->Character->GetWorld()->LineTraceSingleByChannel(OutHit, LambdaInstanceData->Character->GetActorLocation(), SensedActor->GetActorLocation(), ECC_Visibility, QueryParams);
+						// set the task output
+						LambdaInstanceData->TargetActor = SensedActor;
 
-						}
+						// broadcast the see enemy delegate
+						StrongContext.BroadcastDelegate(LambdaInstanceData->OnSeeEnemyDelegate);
+
+					// no direct line of sight to target
+					} else {
 
-						// check if we have a direct line of sight to the stimulus
-						if (bDirectLOS)
+						// if we already have a target, ignore the partial sense and keep on them
+						if (!IsValid(LambdaInstanceData->TargetActor))
 						{
-							// set the controller's target
-							LambdaInstanceData->Controller->SetCurrentTarget(SensedActor);
+							// scale the last stimulus by time elapsed so we phase it out over time
+							const float LastStimulusTime = LambdaInstanceData->Character->GetWorld()->GetTimeSeconds() - LambdaInstanceData->LastStimulusTime;
+							const float ScaledStimulus = LambdaInstanceData->LastStimulusStrength / FMath::Max(LastStimulusTime, 1.0f);
 
-							// set the task output
-							LambdaInstanceData->TargetActor = SensedActor;
+							// is this stimulus stronger?
+							if (Stimulus.Strength > ScaledStimulus)
+							{
+								// update the stimulus strength
+								LambdaInstanceData->LastStimulusStrength = Stimulus.Strength;
 
-							// set the flags
-							LambdaInstanceData->bHasTarget = true;
-							LambdaInstanceData->bHasInvestigateLocation = false;
+								// update the stimulus time
+								LambdaInstanceData->LastStimulusTime = LambdaInstanceData->Character->GetWorld()->GetTimeSeconds();
 
-						// no direct line of sight to target
-						} else {
+								// set the investigate location
+								LambdaInstanceData->InvestigateLocation = Stimulus.StimulusLocation;
 
-							// if we already have a target, ignore the partial sense and keep on them
-							if (!IsValid(LambdaInstanceData->TargetActor))
-							{
-								// is this stimulus stronger than the last one we had?
-								if (Stimulus.Strength > LambdaInstanceData->LastStimulusStrength)
-								{
-									// update the stimulus strength
-									LambdaInstanceData->LastStimulusStrength = Stimulus.Strength;
-
-									// set the investigate location
-									LambdaInstanceData->InvestigateLocation = Stimulus.StimulusLocation;
-
-									// set the investigate flag
-									LambdaInstanceData->bHasInvestigateLocation = true;
-								}
+								// broadcast the investigate delegate
+								StrongContext.BroadcastDelegate(LambdaInstanceData->OnInvestigateLocationDelegate);
 							}
 						}
 					}
 				}
 			}
-		);
+		}
+	);
 
-		// bind the perception forgotten delegate on the controller
-		InstanceData.Controller->OnShooterPerceptionForgotten.BindLambda(
-			[WeakContext = Context.MakeWeakExecutionContext()](AActor* SensedActor)
+	// bind the perception forgotten delegate on the controller
+	InstanceData.Controller->OnShooterPerceptionForgotten.BindLambda(
+		[WeakContext = Context.MakeWeakExecutionContext()](AActor* SensedActor)
+		{
+			// get the instance data inside the lambda
+			const FStateTreeStrongExecutionContext StrongContext = WeakContext.MakeStrongExecutionContext();
+			if (FInstanceDataType* LambdaInstanceData = StrongContext.GetInstanceDataPtr<FInstanceDataType>())
 			{
-				// get the instance data inside the lambda
-				const FStateTreeStrongExecutionContext StrongContext = WeakContext.MakeStrongExecutionContext();
-				if (FInstanceDataType* LambdaInstanceData = StrongContext.GetInstanceDataPtr<FInstanceDataType>())
-				{
-					bool bForget = false;
+				// reset the stimulus strength
+				LambdaInstanceData->LastStimulusStrength = 0.0f;
 
-					// are we forgetting the current target?
-					if (SensedActor == LambdaInstanceData->TargetActor)
+				bool bForget = false;
+
+				// are we forgetting the current target?
+				if (SensedActor == LambdaInstanceData->TargetActor)
+				{
+					bForget = true;
+				}
+				else 
+				{
+					// are we forgetting about a partial sense while we have no target?
+					if (!IsValid(LambdaInstanceData->TargetActor))
 					{
 						bForget = true;
 					}
-					else 
-					{
-						// are we forgetting about a partial sense?
-						if (!IsValid(LambdaInstanceData->TargetActor))
-						{
-							bForget = true;
-						}
-					}
-
-					if (bForget)
-					{
-						// clear the target
-						LambdaInstanceData->TargetActor = nullptr;
+				}
 
-						// clear the flags
-						LambdaInstanceData->bHasInvestigateLocation = false;
-						LambdaInstanceData->bHasTarget = false;
+				if (bForget)
+				{
+					// clear the target
+					LambdaInstanceData->TargetActor = nullptr;
 
-						// reset the stimulus strength
-						LambdaInstanceData->LastStimulusStrength = 0.0f;
+					// clear the target on the controller
+					LambdaInstanceData->Controller->ClearCurrentTarget();
+					LambdaInstanceData->Controller->ClearFocus(EAIFocusPriority::Gameplay);
 
-						// clear the target on the controller
-						LambdaInstanceData->Controller->ClearCurrentTarget();
-						LambdaInstanceData->Controller->ClearFocus(EAIFocusPriority::Gameplay);
-					}
+					// broadcast the forget delegate
+					StrongContext.BroadcastDelegate(LambdaInstanceData->OnForgetEnemyDelegate);
 				}
 			}
-		);
-	}
+		}
+	);
 
 	return EStateTreeRunStatus::Running;
 }
 
 void FStateTreeSenseEnemiesTask::ExitState(FStateTreeExecutionContext& Context, const FStateTreeTransitionResult& Transition) const
 {
-	// have we transitioned to another state?
-	if (Transition.ChangeType == EStateTreeStateChangeType::Changed)
-	{
-		// get the instance data
-		FInstanceDataType& InstanceData = Context.GetInstanceData(*this);
+	// get the instance data
+	FInstanceDataType& InstanceData = Context.GetInstanceData(*this);
 
-		// unbind the perception delegates
-		InstanceData.Controller->OnShooterPerceptionUpdated.Unbind();
-		InstanceData.Controller->OnShooterPerceptionForgotten.Unbind();
-	}
+	// unbind the perception delegates
+	InstanceData.Controller->OnShooterPerceptionUpdated.Unbind();
+	InstanceData.Controller->OnShooterPerceptionForgotten.Unbind();
 }
 
 #if WITH_EDITOR
diff --git a/Templates/TP_FirstPerson/Source/TP_FirstPerson/Variant_Shooter/AI/ShooterStateTreeUtility.h b/Templates/TP_FirstPerson/Source/TP_FirstPerson/Variant_Shooter/AI/ShooterStateTreeUtility.h
--- a/Templates/TP_FirstPerson/Source/TP_FirstPerson/Variant_Shooter/AI/ShooterStateTreeUtility.h
+++ b/Templates/TP_FirstPerson/Source/TP_FirstPerson/Variant_Shooter/AI/ShooterStateTreeUtility.h
@@ -93,6 +93,16 @@ struct FStateTreeFaceActorTask : public FStateTreeTaskCommonBase
 {
 	GENERATED_BODY()
 
+	/** Constructor */
+	FStateTreeFaceActorTask()
+	{
+		// disable tick
+		bShouldCallTick = false;
+
+		// skip state change events if this is sustained
+		bShouldStateChangeOnReselect = false;
+	}
+
 	/* Ensure we're using the correct instance data struct */
 	using FInstanceDataType = FStateTreeFaceActorInstanceData;
 	virtual const UStruct* GetInstanceDataType() const override { return FInstanceDataType::StaticStruct(); }
@@ -135,6 +145,16 @@ struct FStateTreeFaceLocationTask : public FStateTreeTaskCommonBase
 {
 	GENERATED_BODY()
 
+	/** Constructor */
+	FStateTreeFaceLocationTask()
+	{
+		// disable tick
+		bShouldCallTick = false;
+
+		// skip state change events if this is sustained
+		bShouldStateChangeOnReselect = false;
+	}
+
 	/* Ensure we're using the correct instance data struct */
 	using FInstanceDataType = FStateTreeFaceLocationInstanceData;
 	virtual const UStruct* GetInstanceDataType() const override { return FInstanceDataType::StaticStruct(); }
@@ -181,6 +201,16 @@ struct FStateTreeSetRandomFloatTask : public FStateTreeTaskCommonBase
 {
 	GENERATED_BODY()
 
+	/** Constructor */
+	FStateTreeSetRandomFloatTask()
+	{
+		// disable tick
+		bShouldCallTick = false;
+
+		// skip state change events if this is sustained
+		bShouldStateChangeOnReselect = false;
+	}
+
 	/* Ensure we're using the correct instance data struct */
 	using FInstanceDataType = FStateTreeSetRandomFloatData;
 	virtual const UStruct* GetInstanceDataType() const override { return FInstanceDataType::StaticStruct(); }
@@ -220,6 +250,16 @@ struct FStateTreeShootAtTargetTask : public FStateTreeTaskCommonBase
 {
 	GENERATED_BODY()
 
+	/** Constructor */
+	FStateTreeShootAtTargetTask()
+	{
+		// disable tick
+		bShouldCallTick = false;
+
+		// skip state change events if this is sustained
+		bShouldStateChangeOnReselect = false;
+	}
+
 	/* Ensure we're using the correct instance data struct */
 	using FInstanceDataType = FStateTreeShootAtTargetInstanceData;
 	virtual const UStruct* GetInstanceDataType() const override { return FInstanceDataType::StaticStruct(); }
@@ -261,18 +301,22 @@ struct FStateTreeSenseEnemiesInstanceData
 	UPROPERTY(EditAnywhere, Category = Output)
 	FVector InvestigateLocation = FVector::ZeroVector;
 
-	/** True if a target was successfully sensed */
-	UPROPERTY(EditAnywhere, Category = Output)
-	bool bHasTarget = false;
-
-	/** True if an investigate location was successfully sensed */
-	UPROPERTY(EditAnywhere, Category = Output)
-	bool bHasInvestigateLocation = false;
-
 	/** Tag required on sensed actors */
 	UPROPERTY(EditAnywhere, Category = Parameter)
 	FName SenseTag = FName("Player");
 
+	/** StateTree delegate to broadcast when the AI should move to investigate something suspicious */
+	UPROPERTY(EditAnywhere)
+	FStateTreeDelegateDispatcher OnInvestigateLocationDelegate;
+
+	/** StateTree delegate to broadcast when the AI has detected an enemy and should attack it */
+	UPROPERTY(EditAnywhere)
+	FStateTreeDelegateDispatcher OnSeeEnemyDelegate;
+
+	/** StateTree delegate to broadcast when the AI has lost track of its target */
+	UPROPERTY(EditAnywhere)
+	FStateTreeDelegateDispatcher OnForgetEnemyDelegate;
+
 	/** Line of sight cone half angle to consider a full sense */
 	UPROPERTY(EditAnywhere, Category = Parameter)
 	float DirectLineOfSightCone = 85.0f;
@@ -280,6 +324,9 @@ struct FStateTreeSenseEnemiesInstanceData
 	/** Strength of the last processed stimulus */
 	UPROPERTY(EditAnywhere)
 	float LastStimulusStrength = 0.0f;
+
+	UPROPERTY(EditAnywhere)
+	float LastStimulusTime = 0.0f;
 };
 
 /**
@@ -290,6 +337,16 @@ struct FStateTreeSenseEnemiesTask : public FStateTreeTaskCommonBase
 {
 	GENERATED_BODY()
 
+	/** Constructor */
+	FStateTreeSenseEnemiesTask()
+	{
+		// disable tick
+		bShouldCallTick = false;
+
+		// skip state change events if this is sustained
+		bShouldStateChangeOnReselect = false;
+	}
+
 	/* Ensure we're using the correct instance data struct */
 	using FInstanceDataType = FStateTreeSenseEnemiesInstanceData;
 	virtual const UStruct* GetInstanceDataType() const override { return FInstanceDataType::StaticStruct(); }
diff --git a/Templates/TP_FirstPerson/Source/TP_FirstPerson/Variant_Shooter/ShooterCharacter.cpp b/Templates/TP_FirstPerson/Source/TP_FirstPerson/Variant_Shooter/ShooterCharacter.cpp
--- a/Templates/TP_FirstPerson/Source/TP_FirstPerson/Variant_Shooter/ShooterCharacter.cpp
+++ b/Templates/TP_FirstPerson/Source/TP_FirstPerson/Variant_Shooter/ShooterCharacter.cpp
@@ -8,6 +8,7 @@
 #include "Components/PawnNoiseEmitterComponent.h"
 #include "GameFramework/CharacterMovementComponent.h"
 #include "Components/SkeletalMeshComponent.h"
+#include "Components/CapsuleComponent.h"
 #include "Engine/World.h"
 #include "Camera/CameraComponent.h"
 #include "TimerManager.h"
@@ -162,7 +163,7 @@ void AShooterCharacter::DoSwitchWeapon()
 		CurrentWeapon = OwnedWeapons[WeaponIndex];
 
 		// activate the new weapon
-		CurrentWeapon->ActivateWeapon();
+		CurrentWeapon->ActivateWeapon(PlayerTag);
 	}
 }
 
@@ -241,7 +242,7 @@ void AShooterCharacter::AddWeaponClass(const TSubclassOf<AShooterWeapon>& Weapon
 
 			// switch to the new weapon
 			CurrentWeapon = AddedWeapon;
-			CurrentWeapon->ActivateWeapon();
+			CurrentWeapon->ActivateWeapon(PlayerTag);
 		}
 	}
 }
@@ -302,6 +303,9 @@ void AShooterCharacter::Die()
 	// stop character movement
 	GetCharacterMovement()->StopMovementImmediately();
 
+	// disable collision
+	GetCapsuleComponent()->SetCollisionEnabled(ECollisionEnabled::NoCollision);
+
 	// disable controls
 	DisableInput(nullptr);
 
diff --git a/Templates/TP_FirstPerson/Source/TP_FirstPerson/Variant_Shooter/ShooterCharacter.h b/Templates/TP_FirstPerson/Source/TP_FirstPerson/Variant_Shooter/ShooterCharacter.h
--- a/Templates/TP_FirstPerson/Source/TP_FirstPerson/Variant_Shooter/ShooterCharacter.h
+++ b/Templates/TP_FirstPerson/Source/TP_FirstPerson/Variant_Shooter/ShooterCharacter.h
@@ -66,6 +66,10 @@ protected:
 	UPROPERTY(EditAnywhere, Category="Team")
 	FName DeathTag = FName("Dead");
 
+	/** Tag to pass to weapons and projectiles to identify their AI perception noise as player-generated */
+	UPROPERTY(EditAnywhere, Category="Tags")
+	FName PlayerTag = FName("Player");
+
 	/** List of weapons picked up by the character */
 	TArray<AShooterWeapon*> OwnedWeapons;
 
diff --git a/Templates/TP_FirstPerson/Source/TP_FirstPerson/Variant_Shooter/Weapons/ShooterProjectile.cpp b/Templates/TP_FirstPerson/Source/TP_FirstPerson/Variant_Shooter/Weapons/ShooterProjectile.cpp
--- a/Templates/TP_FirstPerson/Source/TP_FirstPerson/Variant_Shooter/Weapons/ShooterProjectile.cpp
+++ b/Templates/TP_FirstPerson/Source/TP_FirstPerson/Variant_Shooter/Weapons/ShooterProjectile.cpp
@@ -165,3 +165,8 @@ void AShooterProjectile::OnDeferredDestruction()
 	// destroy this actor
 	Destroy();
 }
+
+void AShooterProjectile::SetNoiseTag(const FName& Tag)
+{
+	NoiseTag = Tag;
+}
diff --git a/Templates/TP_FirstPerson/Source/TP_FirstPerson/Variant_Shooter/Weapons/ShooterProjectile.h b/Templates/TP_FirstPerson/Source/TP_FirstPerson/Variant_Shooter/Weapons/ShooterProjectile.h
--- a/Templates/TP_FirstPerson/Source/TP_FirstPerson/Variant_Shooter/Weapons/ShooterProjectile.h
+++ b/Templates/TP_FirstPerson/Source/TP_FirstPerson/Variant_Shooter/Weapons/ShooterProjectile.h
@@ -35,7 +35,7 @@ protected:
 
 	/** Range of the AI perception noise done by this projectile on hit */
 	UPROPERTY(EditAnywhere, Category="Projectile|Noise", meta = (ClampMin = 0, ClampMax = 100000, Units = "cm"))
-	float NoiseRange = 3000.0f;
+	float NoiseRange = 1000.0f;
 
 	/** Tag of the AI perception noise done by this projectile on hit */
 	UPROPERTY(EditAnywhere, Category="Noise")
@@ -106,4 +106,9 @@ protected:
 	/** Called from the destruction timer to destroy this projectile */
 	void OnDeferredDestruction();
 
+public:
+
+	/** Sets the noise tag to use when generating AI perception noise on impact */
+	void SetNoiseTag(const FName& Tag);
+
 };
diff --git a/Templates/TP_FirstPerson/Source/TP_FirstPerson/Variant_Shooter/Weapons/ShooterWeapon.cpp b/Templates/TP_FirstPerson/Source/TP_FirstPerson/Variant_Shooter/Weapons/ShooterWeapon.cpp
--- a/Templates/TP_FirstPerson/Source/TP_FirstPerson/Variant_Shooter/Weapons/ShooterWeapon.cpp
+++ b/Templates/TP_FirstPerson/Source/TP_FirstPerson/Variant_Shooter/Weapons/ShooterWeapon.cpp
@@ -68,8 +68,11 @@ void AShooterWeapon::OnOwnerDestroyed(AActor* DestroyedActor)
 	Destroy();
 }
 
-void AShooterWeapon::ActivateWeapon()
+void AShooterWeapon::ActivateWeapon(const FName& OwnerTag)
 {
+	// save the owner tag for perception noise detection
+	NoiseOwnerTag = OwnerTag;
+
 	// unhide this weapon
 	SetActorHiddenInGame(false);
 
@@ -138,7 +141,7 @@ void AShooterWeapon::Fire()
 	TimeOfLastShot = GetWorld()->GetTimeSeconds();
 
 	// make noise so the AI perception system can hear us
-	MakeNoise(ShotLoudness, PawnOwner, PawnOwner->GetActorLocation(), ShotNoiseRange, ShotNoiseTag);
+	MakeNoise(ShotLoudness, PawnOwner, PawnOwner->GetActorLocation(), ShotNoiseRange, NoiseOwnerTag);
 
 	// are we full auto?
 	if (bFullAuto)
@@ -173,6 +176,12 @@ void AShooterWeapon::FireProjectile(const FVector& TargetLocation)
 
 	AShooterProjectile* Projectile = GetWorld()->SpawnActor<AShooterProjectile>(ProjectileClass, ProjectileTransform, SpawnParams);
 
+	// set the noise tag on the projectile
+	if (Projectile)
+	{
+		Projectile->SetNoiseTag(NoiseOwnerTag);
+	}
+
 	// play the firing montage
 	WeaponOwner->PlayFiringMontage(FiringMontage);
 
diff --git a/Templates/TP_FirstPerson/Source/TP_FirstPerson/Variant_Shooter/Weapons/ShooterWeapon.h b/Templates/TP_FirstPerson/Source/TP_FirstPerson/Variant_Shooter/Weapons/ShooterWeapon.h
--- a/Templates/TP_FirstPerson/Source/TP_FirstPerson/Variant_Shooter/Weapons/ShooterWeapon.h
+++ b/Templates/TP_FirstPerson/Source/TP_FirstPerson/Variant_Shooter/Weapons/ShooterWeapon.h
@@ -103,11 +103,11 @@ protected:
 
 	/** Max range of shot AI perception noise */
 	UPROPERTY(EditAnywhere, Category="Perception", meta = (ClampMin = 0, ClampMax = 100000, Units = "cm"))
-	float ShotNoiseRange = 3000.0f;
+	float ShotNoiseRange = 300.0f;
 
 	/** Tag to apply to noise generated by shooting this weapon */
 	UPROPERTY(EditAnywhere, Category="Perception")
-	FName ShotNoiseTag = FName("Shot");
+	FName NoiseOwnerTag = FName("Shot");
 
 public:	
 
@@ -131,7 +131,7 @@ protected:
 public:
 
 	/** Activates this weapon and gets it ready to fire */
-	void ActivateWeapon();
+	void ActivateWeapon(const FName& OwnerTag);
 
 	/** Deactivates this weapon */
 	void DeactivateWeapon();
